Copying //UE4/Dev-Sequencer to Dev-Main (//UE4/Dev-Main)

diff --git a/Engine/Source/Editor/MovieSceneCaptureDialog/Private/MovieSceneCaptureDialogModule.cpp b/Engine/Source/Editor/MovieSceneCaptureDialog/Private/MovieSceneCaptureDialogModule.cpp
--- a/Engine/Source/Editor/MovieSceneCaptureDialog/Private/MovieSceneCaptureDialogModule.cpp
+++ b/Engine/Source/Editor/MovieSceneCaptureDialog/Private/MovieSceneCaptureDialogModule.cpp
@@ -2,14 +2,19 @@
 
 #include "MovieSceneCaptureDialogModule.h"
 #include "MovieSceneCapture.h"
+
+#include "UnrealEdMisc.h"
+
+#include "SlateBasics.h"
+#include "SlateExtras.h"
+#include "SceneViewport.h"
+
 #include "SDockTab.h"
 #include "JsonObjectConverter.h"
 #include "INotificationWidget.h"
 #include "SNotificationList.h"
 #include "NotificationManager.h"
 
-#include "SlateExtras.h"
-
 #include "EditorStyle.h"
 #include "Editor.h"
 #include "PropertyEditing.h"
@@ -20,6 +25,8 @@
 #include "ISessionInfo.h"
 #include "ISessionManager.h"
 
+#include "SLevelViewport.h"
+
 #define LOCTEXT_NAMESPACE "MovieSceneCaptureDialog"
 
 const TCHAR* MovieCaptureSessionName = TEXT("Movie Scene Capture");
@@ -119,26 +126,43 @@ private:
 	UMovieSceneCapture* MovieSceneCapture;
 };
 
-DECLARE_DELEGATE_OneParam(FOnProcessClosed, int32);
+enum class ECaptureState
+{
+	Pending,
+	Success,
+	Failure
+};
+
+DECLARE_DELEGATE_OneParam(FOnCaptureFinished, bool /*bCancelled*/);
 
 class SCaptureMovieNotification : public SCompoundWidget, public INotificationWidget
 {
 public:
 	SLATE_BEGIN_ARGS(SCaptureMovieNotification){}
-		SLATE_EVENT(FOnProcessClosed, OnProcessClosed)
-		SLATE_ARGUMENT(FString, BrowseToFolder)
+
+		SLATE_ATTRIBUTE(ECaptureState, CaptureState)
+
+		SLATE_EVENT(FOnCaptureFinished, OnCaptureFinished)
+
+		SLATE_EVENT(FSimpleDelegate, OnCancel)
+
+		SLATE_ARGUMENT(FString, CapturePath)
+
 	SLATE_END_ARGS()
 
-	void Construct(const FArguments& InArgs, FProcHandle InProcHandle)
+	void Construct(const FArguments& InArgs)
 	{
-		OnProcessClosed = InArgs._OnProcessClosed;
-		ProcHandle = InProcHandle;
+		CaptureState = InArgs._CaptureState;
+		OnCaptureFinished = InArgs._OnCaptureFinished;
+		OnCancel = InArgs._OnCancel;
 
-		FString BrowseToFolder = FPaths::ConvertRelativePathToFull(InArgs._BrowseToFolder);
-		BrowseToFolder.RemoveFromEnd(TEXT("\\"));
+		CachedState = ECaptureState::Pending;
+
+		FString CapturePath = FPaths::ConvertRelativePathToFull(InArgs._CapturePath);
+		CapturePath.RemoveFromEnd(TEXT("\\"));
 
 		auto OnBrowseToFolder = [=]{
-			FPlatformProcess::ExploreFolder(*BrowseToFolder);
+			FPlatformProcess::ExploreFolder(*CapturePath);
 		};
 
 		ChildSlot
@@ -208,15 +232,25 @@ public:
 			return;
 		}
 
-		if (!ProcHandle.IsValid() || !FPlatformProcess::IsProcRunning(ProcHandle))
-		{
-			int32 RetCode = 0;
-			FPlatformProcess::GetProcReturnCode(ProcHandle, &RetCode);
-			OnProcessClosed.ExecuteIfBound(RetCode);
+		ECaptureState StateThisFrame = CaptureState.Get();
 
-			if (RetCode == 0)
+		if (CachedState != StateThisFrame)
+		{
+			CachedState = StateThisFrame;
+			
+			if (CachedState == ECaptureState::Success)
 			{
-				TextBlock->SetText(LOCTEXT("Finished", "Capture Finished"));
+				TextBlock->SetText(LOCTEXT("CaptureFinished", "Capture Finished"));
+				OnCaptureFinished.ExecuteIfBound(true);
+			}
+			else if (CachedState == ECaptureState::Failure)
+			{
+				TextBlock->SetText(LOCTEXT("CaptureFailed", "Capture Failed"));
+				OnCaptureFinished.ExecuteIfBound(false);
+			}
+			else
+			{
+				ensureMsgf(false, TEXT("Cannot move from a finished to a pending state."));
 			}
 		}
 	}
@@ -243,44 +277,158 @@ private:
 	{
 		if (State == SNotificationItem::CS_Pending)
 		{
-			bool bFoundInstance = false;
+			OnCancel.ExecuteIfBound();
+		}
+		return FReply::Handled();
+	}
+	
+private:
+	TSharedPtr<SWidget> Button, Throbber, Hyperlink;
+	TSharedPtr<STextBlock> TextBlock;
+	SNotificationItem::ECompletionState State;
+
+	FSimpleDelegate OnCancel;
+	ECaptureState CachedState;
+	TAttribute<ECaptureState> CaptureState;
+	FOnCaptureFinished OnCaptureFinished;
+};
 
-			// Attempt to send a remote command to gracefully terminate the process
-			ISessionServicesModule& SessionServices = FModuleManager::Get().LoadModuleChecked<ISessionServicesModule>("SessionServices");
-			TSharedRef<ISessionManager> SessionManager = SessionServices.GetSessionManager();
+struct FInEditorCapture
+{
+	FInEditorCapture(UMovieSceneCapture* InCaptureObject, TFunction<void()> InOnStarted)
+		: CaptureObject(InCaptureObject)
+	{
+		ULevelEditorPlaySettings* PlayInEditorSettings = GetMutableDefault<ULevelEditorPlaySettings>();
 
-			TArray<TSharedPtr<ISessionInfo>> Sessions;
-			SessionManager->GetSessions(Sessions);
+		bScreenMessagesWereEnabled = GAreScreenMessagesEnabled;
+		GAreScreenMessagesEnabled = false;
+		OnStarted = InOnStarted;
+		FObjectWriter(PlayInEditorSettings, BackedUpPlaySettings);
+		OverridePlaySettings(PlayInEditorSettings);
 
-			for (const TSharedPtr<ISessionInfo>& Session : Sessions)
+		CaptureObject->AddToRoot();
+		CaptureObject->OnCaptureFinished().AddRaw(this, &FInEditorCapture::OnEnd);
+
+		UGameViewportClient::OnViewportCreated().AddRaw(this, &FInEditorCapture::OnStart);
+		FEditorDelegates::EndPIE.AddRaw(this, &FInEditorCapture::OnEndPIE);
+		
+		GEditor->RequestPlaySession(true, nullptr, false);
+	}
+
+	void OverridePlaySettings(ULevelEditorPlaySettings* PlayInEditorSettings)
+	{
+		const FMovieSceneCaptureSettings& Settings = CaptureObject->GetSettings();
+
+		PlayInEditorSettings->NewWindowWidth = Settings.Resolution.ResX;
+		PlayInEditorSettings->NewWindowHeight = Settings.Resolution.ResY;
+		PlayInEditorSettings->CenterNewWindow = true;
+		PlayInEditorSettings->LastExecutedPlayModeType = EPlayModeType::PlayMode_InEditorFloating;
+
+		TSharedRef<SWindow> CustomWindow = SNew(SWindow)
+			.Title(LOCTEXT("MovieRenderPreviewTitle", "Movie Render - Preview"))
+			.AutoCenter(EAutoCenter::PrimaryWorkArea)
+			.UseOSWindowBorder(true)
+			.FocusWhenFirstShown(false)
+			.ActivateWhenFirstShown(false)
+			.HasCloseButton(true)
+			.SupportsMaximize(false)
+			.SupportsMinimize(true)
+			.SizingRule(ESizingRule::FixedSize);
+
+		FSlateApplication::Get().AddWindow(CustomWindow);
+
+		PlayInEditorSettings->CustomPIEWindow = CustomWindow;
+
+		// Reset everything else
+		PlayInEditorSettings->GameGetsMouseControl = false;
+		PlayInEditorSettings->ShowMouseControlLabel = false;
+		PlayInEditorSettings->ViewportGetsHMDControl = false;
+		PlayInEditorSettings->EnableSound = false;
+		PlayInEditorSettings->bOnlyLoadVisibleLevelsInPIE = false;
+		PlayInEditorSettings->bPreferToStreamLevelsInPIE = false;
+		PlayInEditorSettings->PIEAlwaysOnTop = false;
+		PlayInEditorSettings->DisableStandaloneSound = true;
+		PlayInEditorSettings->AdditionalLaunchParameters = TEXT("");
+		PlayInEditorSettings->BuildGameBeforeLaunch = EPlayOnBuildMode::PlayOnBuild_Never;
+		PlayInEditorSettings->LaunchConfiguration = EPlayOnLaunchConfiguration::LaunchConfig_Default;
+		PlayInEditorSettings->SetPlayNetMode(EPlayNetMode::PIE_Standalone);
+		PlayInEditorSettings->SetRunUnderOneProcess(true);
+		PlayInEditorSettings->SetPlayNetDedicated(false);
+		PlayInEditorSettings->SetPlayNumberOfClients(1);
+	}
+
+	void OnStart()
+	{
+		for (const FWorldContext& Context : GEngine->GetWorldContexts())
+		{
+			if (Context.WorldType == EWorldType::PIE)
 			{
-				if (Session->GetSessionName() == MovieCaptureSessionName)
+				FSlatePlayInEditorInfo* SlatePlayInEditorSession = GEditor->SlatePlayInEditorMap.Find(Context.ContextHandle);
+				if (SlatePlayInEditorSession)
 				{
-					TArray<TSharedPtr<ISessionInstanceInfo>> Instances;
-					Session->GetInstances(Instances);
+					TSharedPtr<SWindow> Window = SlatePlayInEditorSession->SlatePlayInEditorWindow.Pin();
 
-					for (const TSharedPtr<ISessionInstanceInfo>& Instance : Instances)
+					const FMovieSceneCaptureSettings& Settings = CaptureObject->GetSettings();
+
+					SlatePlayInEditorSession->SlatePlayInEditorWindowViewport->SetViewportSize(Settings.Resolution.ResX,Settings.Resolution.ResY);
+
+					FVector2D PreviewWindowSize(Settings.Resolution.ResX, Settings.Resolution.ResY);
+
+					// Keep scaling down the window size while we're bigger than half the destop width/height
 					{
-						Instance->ExecuteCommand("exit");
-						bFoundInstance = true;
+						FDisplayMetrics DisplayMetrics;
+						FSlateApplication::Get().GetDisplayMetrics(DisplayMetrics);
+						
+						while(PreviewWindowSize.X >= DisplayMetrics.PrimaryDisplayWidth*.5f || PreviewWindowSize.Y >= DisplayMetrics.PrimaryDisplayHeight*.5f)
+						{
+							PreviewWindowSize *= .5f;
+						}
 					}
-				}
-			}
+					
+					Window->Resize(PreviewWindowSize);
 
-			if (!bFoundInstance)
-			{
-				FPlatformProcess::TerminateProc(ProcHandle);
+					CaptureObject->Initialize(SlatePlayInEditorSession->SlatePlayInEditorWindowViewport, Context.PIEInstance);
+					OnStarted();
+				}
+				return;
 			}
 		}
-		return FReply::Handled();
+
+		// todo: error?
 	}
-	
-private:
-	TSharedPtr<SWidget> Button, Throbber, Hyperlink;
-	TSharedPtr<STextBlock> TextBlock;
-	SNotificationItem::ECompletionState State;
-	FOnProcessClosed OnProcessClosed;
-	FProcHandle ProcHandle;
+
+	void Shutdown()
+	{
+		FEditorDelegates::EndPIE.RemoveAll(this);
+		UGameViewportClient::OnViewportCreated().RemoveAll(this);
+		CaptureObject->OnCaptureFinished().RemoveAll(this);
+
+		GAreScreenMessagesEnabled = bScreenMessagesWereEnabled;
+
+		FObjectReader(GetMutableDefault<ULevelEditorPlaySettings>(), BackedUpPlaySettings);
+
+		CaptureObject->Close();
+		CaptureObject->RemoveFromRoot();
+
+	}
+	void OnEndPIE(bool bIsSimulating)
+	{
+		Shutdown();
+		delete this;
+	}
+
+	void OnEnd()
+	{
+		Shutdown();
+		delete this;
+
+		GEditor->RequestEndPlayMap();
+	}
+
+	TFunction<void()> OnStarted;
+	bool bScreenMessagesWereEnabled;
+	TArray<uint8> BackedUpPlaySettings;
+	UMovieSceneCapture* CaptureObject;
 };
 
 class FMovieSceneCaptureDialogModule : public IMovieSceneCaptureDialogModule
@@ -325,9 +473,9 @@ class FMovieSceneCaptureDialogModule : public IMovieSceneCaptureDialogModule
 		CaptureSettingsWindow = ExistingWindow;
 	}
 
-	void OnMovieCaptureProcessClosed(int32 RetCode)
+	void OnCaptureFinished(bool bSuccess)
 	{
-		if (RetCode == 0)
+		if (bSuccess)
 		{
 			InProgressCaptureNotification->SetCompletionState(SNotificationItem::CS_Success);
 		}
@@ -383,13 +531,54 @@ class FMovieSceneCaptureDialogModule : public IMovieSceneCaptureDialogModule
 			return LOCTEXT("AlreadyCapturing", "There is already a movie scene capture process open. Please close it and try again.");
 		}
 
-		// If buffer visualization dumping is enabled, we need to tell capture process to enable it too
-		static const auto CVarDumpFrames = IConsoleManager::Get().FindTConsoleVariableDataInt(TEXT("r.BufferVisualizationDumpFrames"));
-		if (CVarDumpFrames && CVarDumpFrames->GetValueOnGameThread())
+		CaptureObject->SaveConfig();
+		if (CaptureObject->ProtocolSettings)
 		{
-			CaptureObject->bBufferVisualizationDumpFrames = true;
+			CaptureObject->ProtocolSettings->SaveConfig();
 		}
 
+		return CaptureObject->bUseSeparateProcess ? CaptureInNewProcess(CaptureObject, MapNameToLoad) : CaptureInEditor(CaptureObject, MapNameToLoad);
+	}
+
+	FText CaptureInEditor(UMovieSceneCapture* CaptureObject, const FString& MapNameToLoad)
+	{
+		auto GetCaptureStatus = [=]{
+			for (const FWorldContext& Context : GEngine->GetWorldContexts())
+			{
+				if (Context.WorldType == EWorldType::PIE)
+				{
+					return ECaptureState::Pending;
+				}
+			}
+
+			return ECaptureState::Success;
+		};
+
+		auto OnCaptureStarted = [=]{
+
+			FNotificationInfo Info(
+				SNew(SCaptureMovieNotification)
+				.CaptureState_Lambda(GetCaptureStatus)
+				.CapturePath(CaptureObject->Settings.OutputDirectory.Path)
+				.OnCaptureFinished_Raw(this, &FMovieSceneCaptureDialogModule::OnCaptureFinished)
+				.OnCancel_Lambda([]{
+					GEditor->RequestEndPlayMap();
+				})
+			);
+			Info.bFireAndForget = false;
+			Info.ExpireDuration = 5.f;
+			InProgressCaptureNotification = FSlateNotificationManager::Get().AddNotification(Info);
+			InProgressCaptureNotification->SetCompletionState(SNotificationItem::CS_Pending);
+		};
+
+		// deliberately 'leak' the object, since it owns itself
+		new FInEditorCapture(CaptureObject, OnCaptureStarted);
+
+		return FText();
+	}
+
+	FText CaptureInNewProcess(UMovieSceneCapture* CaptureObject, const FString& MapNameToLoad)
+	{
 		// Save out the capture manifest to json
 		FString Filename = FPaths::GameSavedDir() / TEXT("MovieSceneCapture/Manifest.json");
 
@@ -400,6 +589,16 @@ class FMovieSceneCaptureDialogModule : public IMovieSceneCaptureDialogModule
 			RootObject->SetField(TEXT("Type"), MakeShareable(new FJsonValueString(CaptureObject->GetClass()->GetPathName())));
 			RootObject->SetField(TEXT("Data"), MakeShareable(new FJsonValueObject(Object)));
 
+			if (CaptureObject->ProtocolSettings)
+			{
+				RootObject->SetField(TEXT("ProtocolType"), MakeShareable(new FJsonValueString(CaptureObject->ProtocolSettings->GetClass()->GetPathName())));
+				TSharedRef<FJsonObject> ProtocolDataObject = MakeShareable(new FJsonObject);
+				if (FJsonObjectConverter::UStructToJsonObject(CaptureObject->ProtocolSettings->GetClass(), CaptureObject->ProtocolSettings, ProtocolDataObject, 0, 0))
+				{
+					RootObject->SetField(TEXT("ProtocolData"), MakeShareable(new FJsonValueObject(ProtocolDataObject)));
+				}
+			}
+
 			FString Json;
 			TSharedRef<TJsonWriter<> > JsonWriter = TJsonWriterFactory<>::Create(&Json, 0);
 			if (FJsonSerializer::Serialize( RootObject, JsonWriter ))
@@ -449,8 +648,6 @@ class FMovieSceneCaptureDialogModule : public IMovieSceneCaptureDialogModule
 		// Ensure game session is correctly set up 
 		EditorCommandLine += FString::Printf(TEXT(" -messaging -SessionName=\"%s\""), MovieCaptureSessionName);
 
-		CaptureObject->SaveConfig();
-
 		FString Params;
 		if (FPaths::IsProjectFilePathSet())
 		{
@@ -467,10 +664,56 @@ class FMovieSceneCaptureDialogModule : public IMovieSceneCaptureDialogModule
 		// @todo: progress reporting, UI feedback
 		if (ProcessHandle.IsValid())
 		{
+			TSharedRef<FProcHandle> SharedProcHandle = MakeShareable(new FProcHandle(ProcessHandle));
+			auto GetCaptureStatus = [=]{
+				if (!FPlatformProcess::IsProcRunning(*SharedProcHandle))
+				{
+					int32 RetCode = 0;
+					FPlatformProcess::GetProcReturnCode(*SharedProcHandle, &RetCode);
+					return RetCode == 0 ? ECaptureState::Success : ECaptureState::Failure;
+				}
+				else
+				{
+					return ECaptureState::Pending;
+				}
+			};
+
+			auto OnCancel = [=]{
+				bool bFoundInstance = false;
+
+				// Attempt to send a remote command to gracefully terminate the process
+				ISessionServicesModule& SessionServices = FModuleManager::Get().LoadModuleChecked<ISessionServicesModule>("SessionServices");
+				TSharedRef<ISessionManager> SessionManager = SessionServices.GetSessionManager();
+
+				TArray<TSharedPtr<ISessionInfo>> Sessions;
+				SessionManager->GetSessions(Sessions);
+
+				for (const TSharedPtr<ISessionInfo>& Session : Sessions)
+				{
+					if (Session->GetSessionName() == MovieCaptureSessionName)
+					{
+						TArray<TSharedPtr<ISessionInstanceInfo>> Instances;
+						Session->GetInstances(Instances);
+
+						for (const TSharedPtr<ISessionInstanceInfo>& Instance : Instances)
+						{
+							Instance->ExecuteCommand("exit");
+							bFoundInstance = true;
+						}
+					}
+				}
+
+				if (!bFoundInstance)
+				{
+					FPlatformProcess::TerminateProc(*SharedProcHandle);
+				}
+			};
 			FNotificationInfo Info(
-				SNew(SCaptureMovieNotification, ProcessHandle)
-				.BrowseToFolder(CaptureObject->Settings.OutputDirectory.Path)
-				.OnProcessClosed_Raw(this, &FMovieSceneCaptureDialogModule::OnMovieCaptureProcessClosed)
+				SNew(SCaptureMovieNotification)
+				.CaptureState_Lambda(GetCaptureStatus)
+				.CapturePath(CaptureObject->Settings.OutputDirectory.Path)
+				.OnCaptureFinished_Raw(this, &FMovieSceneCaptureDialogModule::OnCaptureFinished)
+				.OnCancel_Lambda(OnCancel)
 			);
 			Info.bFireAndForget = false;
 			Info.ExpireDuration = 5.f;
diff --git a/Engine/Source/Editor/MovieSceneTools/MovieSceneTools.Build.cs b/Engine/Source/Editor/MovieSceneTools/MovieSceneTools.Build.cs
--- a/Engine/Source/Editor/MovieSceneTools/MovieSceneTools.Build.cs
+++ b/Engine/Source/Editor/MovieSceneTools/MovieSceneTools.Build.cs
@@ -59,7 +59,5 @@ public class MovieSceneTools : ModuleRules
 				"PropertyEditor"
 			}
 		);
-
-        bFasterWithoutUnity = true;
 	}
 }
diff --git a/Engine/Source/Editor/MovieSceneTools/Private/TrackEditors/AttachTrackEditor.cpp b/Engine/Source/Editor/MovieSceneTools/Private/TrackEditors/AttachTrackEditor.cpp
--- a/Engine/Source/Editor/MovieSceneTools/Private/TrackEditors/AttachTrackEditor.cpp
+++ b/Engine/Source/Editor/MovieSceneTools/Private/TrackEditors/AttachTrackEditor.cpp
@@ -113,16 +113,6 @@ TSharedRef<ISequencerSection> F3DAttachTrackEditor::MakeSectionInterface( UMovie
 }
 
 
-void F3DAttachTrackEditor::AddKey( const FGuid& ObjectGuid, UObject* AdditionalAsset )
-{
-	AActor* ParentActor = Cast<AActor>(AdditionalAsset);
-	if (ParentActor != nullptr)
-	{
-		UMovieSceneSection* DummySection = nullptr;
-		ActorPicked(ParentActor, ObjectGuid, DummySection);
-	}
-}
-
 void F3DAttachTrackEditor::BuildObjectBindingTrackMenu(FMenuBuilder& MenuBuilder, const FGuid& ObjectBinding, const UClass* ObjectClass)
 {
 	if (ObjectClass->IsChildOf(AActor::StaticClass()))
diff --git a/Engine/Source/Editor/MovieSceneTools/Private/TrackEditors/EventTrackEditor.cpp b/Engine/Source/Editor/MovieSceneTools/Private/TrackEditors/EventTrackEditor.cpp
--- a/Engine/Source/Editor/MovieSceneTools/Private/TrackEditors/EventTrackEditor.cpp
+++ b/Engine/Source/Editor/MovieSceneTools/Private/TrackEditors/EventTrackEditor.cpp
@@ -29,12 +29,6 @@ FEventTrackEditor::FEventTrackEditor(TSharedRef<ISequencer> InSequencer)
 /* ISequencerTrackEditor interface
  *****************************************************************************/
 
-void FEventTrackEditor::AddKey(const FGuid& ObjectGuid, UObject* AdditionalAsset)
-{
-	// todo gmp: Sequencer: implement event track section
-}
-
-
 void FEventTrackEditor::BuildAddTrackMenu(FMenuBuilder& MenuBuilder)
 {
 	UMovieSceneSequence* RootMovieSceneSequence = GetSequencer()->GetRootMovieSceneSequence();
diff --git a/Engine/Source/Editor/MovieSceneTools/Private/TrackEditors/EventTrackEditor.h b/Engine/Source/Editor/MovieSceneTools/Private/TrackEditors/EventTrackEditor.h
--- a/Engine/Source/Editor/MovieSceneTools/Private/TrackEditors/EventTrackEditor.h
+++ b/Engine/Source/Editor/MovieSceneTools/Private/TrackEditors/EventTrackEditor.h
@@ -38,7 +38,6 @@ public:
 
 	// ISequencerTrackEditor interface
 
-	virtual void AddKey(const FGuid& ObjectGuid, UObject* AdditionalAsset = nullptr) override;
 	virtual void BuildAddTrackMenu(FMenuBuilder& MenuBuilder) override;
 	virtual void BuildObjectBindingTrackMenu(FMenuBuilder& MenuBuilder, const FGuid& ObjectBinding, const UClass* ObjectClass) override;
 	virtual TSharedRef<ISequencerSection> MakeSectionInterface(UMovieSceneSection& SectionObject, UMovieSceneTrack& Track) override;
diff --git a/Engine/Source/Editor/MovieSceneTools/Private/TrackEditors/PathTrackEditor.cpp b/Engine/Source/Editor/MovieSceneTools/Private/TrackEditors/PathTrackEditor.cpp
--- a/Engine/Source/Editor/MovieSceneTools/Private/TrackEditors/PathTrackEditor.cpp
+++ b/Engine/Source/Editor/MovieSceneTools/Private/TrackEditors/PathTrackEditor.cpp
@@ -109,17 +109,6 @@ TSharedRef<ISequencerSection> F3DPathTrackEditor::MakeSectionInterface( UMovieSc
 }
 
 
-void F3DPathTrackEditor::AddKey( const FGuid& ObjectGuid, UObject* AdditionalAsset )
-{
-	AActor* ParentActor = Cast<AActor>(AdditionalAsset);
-	if (ParentActor != nullptr)
-	{
-		UMovieSceneSection* DummySection = nullptr;
-		ActorPicked(ParentActor, ObjectGuid, DummySection);
-	}
-}
-
-
 void F3DPathTrackEditor::BuildObjectBindingTrackMenu(FMenuBuilder& MenuBuilder, const FGuid& ObjectBinding, const UClass* ObjectClass)
 {
 	if (ObjectClass->IsChildOf(AActor::StaticClass()))
diff --git a/Engine/Source/Editor/MovieSceneTools/Private/TrackEditors/PropertyTrackEditors/BoolPropertyTrackEditor.cpp b/Engine/Source/Editor/MovieSceneTools/Private/TrackEditors/PropertyTrackEditors/BoolPropertyTrackEditor.cpp
--- a/Engine/Source/Editor/MovieSceneTools/Private/TrackEditors/PropertyTrackEditors/BoolPropertyTrackEditor.cpp
+++ b/Engine/Source/Editor/MovieSceneTools/Private/TrackEditors/PropertyTrackEditors/BoolPropertyTrackEditor.cpp
@@ -18,7 +18,7 @@ TSharedRef<FPropertySection> FBoolPropertyTrackEditor::MakePropertySectionInterf
 }
 
 
-void FBoolPropertyTrackEditor::GenerateKeysFromPropertyChanged(const FPropertyChangedParams& PropertyChangedParams, TArray<bool>& GeneratedKeys)
+void FBoolPropertyTrackEditor::GenerateKeysFromPropertyChanged( const FPropertyChangedParams& PropertyChangedParams, TArray<bool>& NewGeneratedKeys, TArray<bool>& DefaultGeneratedKeys )
 {
-	GeneratedKeys.Add( PropertyChangedParams.GetPropertyValue<bool>() );
+	NewGeneratedKeys.Add( PropertyChangedParams.GetPropertyValue<bool>() );
 }
diff --git a/Engine/Source/Editor/MovieSceneTools/Private/TrackEditors/PropertyTrackEditors/BoolPropertyTrackEditor.h b/Engine/Source/Editor/MovieSceneTools/Private/TrackEditors/PropertyTrackEditors/BoolPropertyTrackEditor.h
--- a/Engine/Source/Editor/MovieSceneTools/Private/TrackEditors/PropertyTrackEditors/BoolPropertyTrackEditor.h
+++ b/Engine/Source/Editor/MovieSceneTools/Private/TrackEditors/PropertyTrackEditors/BoolPropertyTrackEditor.h
@@ -33,5 +33,5 @@ protected:
 
 	// FPropertyTrackEditor interface
 	virtual TSharedRef<FPropertySection> MakePropertySectionInterface( UMovieSceneSection& SectionObject, UMovieSceneTrack& Track ) override;
-	virtual void GenerateKeysFromPropertyChanged( const FPropertyChangedParams& PropertyChangedParams, TArray<bool>& GeneratedKeys ) override;
+	virtual void GenerateKeysFromPropertyChanged( const FPropertyChangedParams& PropertyChangedParams, TArray<bool>& NewGeneratedKeys, TArray<bool>& DefaultGeneratedKeys ) override;
 };
diff --git a/Engine/Source/Editor/MovieSceneTools/Private/TrackEditors/PropertyTrackEditors/BytePropertyTrackEditor.cpp b/Engine/Source/Editor/MovieSceneTools/Private/TrackEditors/PropertyTrackEditors/BytePropertyTrackEditor.cpp
--- a/Engine/Source/Editor/MovieSceneTools/Private/TrackEditors/PropertyTrackEditors/BytePropertyTrackEditor.cpp
+++ b/Engine/Source/Editor/MovieSceneTools/Private/TrackEditors/PropertyTrackEditors/BytePropertyTrackEditor.cpp
@@ -68,7 +68,7 @@ UMovieSceneTrack* FBytePropertyTrackEditor::AddTrack(UMovieScene* FocusedMovieSc
 }
 
 
-void FBytePropertyTrackEditor::GenerateKeysFromPropertyChanged( const FPropertyChangedParams& PropertyChangedParams, TArray<uint8>& GeneratedKeys )
+void FBytePropertyTrackEditor::GenerateKeysFromPropertyChanged( const FPropertyChangedParams& PropertyChangedParams, TArray<uint8>& NewGeneratedKeys, TArray<uint8>& DefaultGeneratedKeys )
 {
-	GeneratedKeys.Add( PropertyChangedParams.GetPropertyValue<uint8>() );
+	NewGeneratedKeys.Add( PropertyChangedParams.GetPropertyValue<uint8>() );
 }
diff --git a/Engine/Source/Editor/MovieSceneTools/Private/TrackEditors/PropertyTrackEditors/BytePropertyTrackEditor.h b/Engine/Source/Editor/MovieSceneTools/Private/TrackEditors/PropertyTrackEditors/BytePropertyTrackEditor.h
--- a/Engine/Source/Editor/MovieSceneTools/Private/TrackEditors/PropertyTrackEditors/BytePropertyTrackEditor.h
+++ b/Engine/Source/Editor/MovieSceneTools/Private/TrackEditors/PropertyTrackEditors/BytePropertyTrackEditor.h
@@ -42,5 +42,5 @@ protected:
 	// FPropertyTrackEditor interface
 
 	virtual TSharedRef<FPropertySection> MakePropertySectionInterface( UMovieSceneSection& SectionObject, UMovieSceneTrack& Track ) override;
-	virtual void GenerateKeysFromPropertyChanged( const FPropertyChangedParams& PropertyChangedParams, TArray<uint8>& GeneratedKeys ) override;
+	virtual void GenerateKeysFromPropertyChanged( const FPropertyChangedParams& PropertyChangedParams, TArray<uint8>& NewGeneratedKeys, TArray<uint8>& DefaultGeneratedKeys ) override;
 };
diff --git a/Engine/Source/Editor/MovieSceneTools/Private/TrackEditors/PropertyTrackEditors/ColorPropertyTrackEditor.cpp b/Engine/Source/Editor/MovieSceneTools/Private/TrackEditors/PropertyTrackEditors/ColorPropertyTrackEditor.cpp
--- a/Engine/Source/Editor/MovieSceneTools/Private/TrackEditors/PropertyTrackEditors/ColorPropertyTrackEditor.cpp
+++ b/Engine/Source/Editor/MovieSceneTools/Private/TrackEditors/PropertyTrackEditors/ColorPropertyTrackEditor.cpp
@@ -24,7 +24,7 @@ TSharedRef<FPropertySection> FColorPropertyTrackEditor::MakePropertySectionInter
 }
 
 
-void FColorPropertyTrackEditor::GenerateKeysFromPropertyChanged( const FPropertyChangedParams& PropertyChangedParams, TArray<FColorKey>& GeneratedKeys )
+void FColorPropertyTrackEditor::GenerateKeysFromPropertyChanged( const FPropertyChangedParams& PropertyChangedParams, TArray<FColorKey>& NewGeneratedKeys, TArray<FColorKey>& DefaultGeneratedKeys )
 {
 	const UStructProperty* StructProp = Cast<const UStructProperty>( PropertyChangedParams.PropertyPath.Last() );
 	FName StructName = StructProp->Struct->GetFName();
@@ -52,20 +52,15 @@ void FColorPropertyTrackEditor::GenerateKeysFromPropertyChanged( const FProperty
 
 	FName ChannelName = PropertyChangedParams.StructPropertyNameToKey;
 
-	if ( ChannelName == NAME_None || ChannelName == RedName )
-	{
-		GeneratedKeys.Add( FColorKey( EKeyColorChannel::Red, ColorValue.R, bIsSlateColor ) );
-	}
-	if ( ChannelName == NAME_None || ChannelName == GreenName )
-	{
-		GeneratedKeys.Add( FColorKey( EKeyColorChannel::Green, ColorValue.G, bIsSlateColor ) );
-	}
-	if ( ChannelName == NAME_None || ChannelName == BlueName )
-	{
-		GeneratedKeys.Add( FColorKey( EKeyColorChannel::Blue, ColorValue.B, bIsSlateColor ) );
-	}
-	if ( ChannelName == NAME_None || ChannelName == AlphaName )
-	{
-		GeneratedKeys.Add( FColorKey( EKeyColorChannel::Alpha, ColorValue.A, bIsSlateColor ) );
-	}
+	TArray<FColorKey>& RedKeys = ChannelName == NAME_None || ChannelName == RedName ? NewGeneratedKeys : DefaultGeneratedKeys;
+	RedKeys.Add( FColorKey( EKeyColorChannel::Red, ColorValue.R, bIsSlateColor ) );
+
+	TArray<FColorKey>& GreenKeys = ChannelName == NAME_None || ChannelName == GreenName ? NewGeneratedKeys : DefaultGeneratedKeys;
+	GreenKeys.Add( FColorKey( EKeyColorChannel::Green, ColorValue.G, bIsSlateColor ) );
+
+	TArray<FColorKey>& BlueKeys =  ChannelName == NAME_None || ChannelName == BlueName ? NewGeneratedKeys : DefaultGeneratedKeys;
+	BlueKeys.Add( FColorKey( EKeyColorChannel::Blue, ColorValue.B, bIsSlateColor ) );
+
+	TArray<FColorKey>& AlphaKeys = ChannelName == NAME_None || ChannelName == AlphaName ? NewGeneratedKeys : DefaultGeneratedKeys;
+	AlphaKeys.Add( FColorKey( EKeyColorChannel::Alpha, ColorValue.A, bIsSlateColor ) );
 }
diff --git a/Engine/Source/Editor/MovieSceneTools/Private/TrackEditors/PropertyTrackEditors/ColorPropertyTrackEditor.h b/Engine/Source/Editor/MovieSceneTools/Private/TrackEditors/PropertyTrackEditors/ColorPropertyTrackEditor.h
--- a/Engine/Source/Editor/MovieSceneTools/Private/TrackEditors/PropertyTrackEditors/ColorPropertyTrackEditor.h
+++ b/Engine/Source/Editor/MovieSceneTools/Private/TrackEditors/PropertyTrackEditors/ColorPropertyTrackEditor.h
@@ -36,7 +36,7 @@ protected:
 	// FPropertyTrackEditor interface
 
 	virtual TSharedRef<FPropertySection> MakePropertySectionInterface( UMovieSceneSection& SectionObject, UMovieSceneTrack& Track ) override;
-	virtual void GenerateKeysFromPropertyChanged( const FPropertyChangedParams& PropertyChangedParams, TArray<FColorKey>& GeneratedKeys ) override;
+	virtual void GenerateKeysFromPropertyChanged( const FPropertyChangedParams& PropertyChangedParams, TArray<FColorKey>& NewGeneratedKeys, TArray<FColorKey>& DefaultGeneratedKeys ) override;
 
 private:
 	static FName RedName;
diff --git a/Engine/Source/Editor/MovieSceneTools/Private/TrackEditors/PropertyTrackEditors/FloatPropertyTrackEditor.cpp b/Engine/Source/Editor/MovieSceneTools/Private/TrackEditors/PropertyTrackEditors/FloatPropertyTrackEditor.cpp
--- a/Engine/Source/Editor/MovieSceneTools/Private/TrackEditors/PropertyTrackEditors/FloatPropertyTrackEditor.cpp
+++ b/Engine/Source/Editor/MovieSceneTools/Private/TrackEditors/PropertyTrackEditors/FloatPropertyTrackEditor.cpp
@@ -19,9 +19,9 @@ TSharedRef<FPropertySection> FFloatPropertyTrackEditor::MakePropertySectionInter
 }
 
 
-void FFloatPropertyTrackEditor::GenerateKeysFromPropertyChanged( const FPropertyChangedParams& PropertyChangedParams, TArray<float>& GeneratedKeys )
+void FFloatPropertyTrackEditor::GenerateKeysFromPropertyChanged( const FPropertyChangedParams& PropertyChangedParams, TArray<float>& NewGeneratedKeys, TArray<float>& DefaultGeneratedKeys )
 {
-	GeneratedKeys.Add( PropertyChangedParams.GetPropertyValue<float>() );
+	NewGeneratedKeys.Add( PropertyChangedParams.GetPropertyValue<float>() );
 }
 
 
diff --git a/Engine/Source/Editor/MovieSceneTools/Private/TrackEditors/PropertyTrackEditors/FloatPropertyTrackEditor.h b/Engine/Source/Editor/MovieSceneTools/Private/TrackEditors/PropertyTrackEditors/FloatPropertyTrackEditor.h
--- a/Engine/Source/Editor/MovieSceneTools/Private/TrackEditors/PropertyTrackEditors/FloatPropertyTrackEditor.h
+++ b/Engine/Source/Editor/MovieSceneTools/Private/TrackEditors/PropertyTrackEditors/FloatPropertyTrackEditor.h
@@ -38,5 +38,5 @@ protected:
 	// FPropertyTrackEditor interface
 
 	virtual TSharedRef<FPropertySection> MakePropertySectionInterface( UMovieSceneSection& SectionObject, UMovieSceneTrack& Track ) override;
-	virtual void GenerateKeysFromPropertyChanged( const FPropertyChangedParams& PropertyChangedParams, TArray<float>& GeneratedKeys ) override;
+	virtual void GenerateKeysFromPropertyChanged( const FPropertyChangedParams& PropertyChangedParams, TArray<float>& NewGeneratedKeys, TArray<float>& DefaultGeneratedKeys ) override;
 };
diff --git a/Engine/Source/Editor/MovieSceneTools/Private/TrackEditors/PropertyTrackEditors/VectorPropertyTrackEditor.cpp b/Engine/Source/Editor/MovieSceneTools/Private/TrackEditors/PropertyTrackEditors/VectorPropertyTrackEditor.cpp
--- a/Engine/Source/Editor/MovieSceneTools/Private/TrackEditors/PropertyTrackEditors/VectorPropertyTrackEditor.cpp
+++ b/Engine/Source/Editor/MovieSceneTools/Private/TrackEditors/PropertyTrackEditors/VectorPropertyTrackEditor.cpp
@@ -22,8 +22,7 @@ TSharedRef<FPropertySection> FVectorPropertyTrackEditor::MakePropertySectionInte
 	return MakeShareable(new FVectorPropertySection(SectionObject, Track.GetDisplayName()));
 }
 
-
-void FVectorPropertyTrackEditor::GenerateKeysFromPropertyChanged( const FPropertyChangedParams& PropertyChangedParams, TArray<FVectorKey>& GeneratedKeys )
+void FVectorPropertyTrackEditor::GenerateKeysFromPropertyChanged( const FPropertyChangedParams& PropertyChangedParams, TArray<FVectorKey>& NewGeneratedKeys, TArray<FVectorKey>& DefaultGeneratedKeys )
 {
 	const UStructProperty* StructProp = Cast<const UStructProperty>( PropertyChangedParams.PropertyPath.Last() );
 
@@ -58,16 +57,23 @@ void FVectorPropertyTrackEditor::GenerateKeysFromPropertyChanged( const FPropert
 	}
 
 	FName ChannelName = PropertyChangedParams.StructPropertyNameToKey;
-		
-	GeneratedKeys.Add( FVectorKey( EKeyVectorChannel::X, VectorValues.X, ChannelName == NAME_None || ChannelName == XName ? EKeyVectorValueType::Key : EKeyVectorValueType::Default ) );
-	GeneratedKeys.Add( FVectorKey( EKeyVectorChannel::Y, VectorValues.Y, ChannelName == NAME_None || ChannelName == YName ? EKeyVectorValueType::Key : EKeyVectorValueType::Default ) );
-	if (Channels >= 3 )
+
+	TArray<FVectorKey>& XKeys = ChannelName == NAME_None || ChannelName == XName ? NewGeneratedKeys : DefaultGeneratedKeys;
+	XKeys.Add( FVectorKey( EKeyVectorChannel::X, VectorValues.X ) );
+
+	TArray<FVectorKey>& YKeys = ChannelName == NAME_None || ChannelName == YName ? NewGeneratedKeys : DefaultGeneratedKeys;
+	YKeys.Add( FVectorKey( EKeyVectorChannel::Y, VectorValues.Y ) );
+
+	if ( Channels >= 3 )
 	{
-		GeneratedKeys.Add( FVectorKey( EKeyVectorChannel::Z, VectorValues.Z, ChannelName == NAME_None || ChannelName == ZName ? EKeyVectorValueType::Key : EKeyVectorValueType::Default ) );
+		TArray<FVectorKey>& ZKeys = ChannelName == NAME_None || ChannelName == ZName ? NewGeneratedKeys : DefaultGeneratedKeys;
+		ZKeys.Add( FVectorKey( EKeyVectorChannel::Z, VectorValues.Z ) );
 	}
-	if (Channels >= 4 )
+
+	if ( Channels >= 4 )
 	{
-		GeneratedKeys.Add( FVectorKey( EKeyVectorChannel::W, VectorValues.W, ChannelName == NAME_None || ChannelName == WName ? EKeyVectorValueType::Key : EKeyVectorValueType::Default ) );
+		TArray<FVectorKey>& WKeys = ChannelName == NAME_None || ChannelName == WName ? NewGeneratedKeys : DefaultGeneratedKeys;
+		WKeys.Add( FVectorKey( EKeyVectorChannel::W, VectorValues.W ) );
 	}
 }
 
@@ -90,8 +96,3 @@ void FVectorPropertyTrackEditor::InitializeNewTrack( UMovieSceneVectorTrack* New
 		NewTrack->SetNumChannelsUsed( 4 );
 	}
 }
-
-bool FVectorPropertyTrackEditor::ShouldAddKey(UMovieSceneVectorTrack* InTrack, FVectorKey InKey, FKeyParams InKeyParams) const
-{
-	return FPropertyTrackEditor::ShouldAddKey(InTrack, InKey, InKeyParams) && InKey.ValueType == EKeyVectorValueType::Key;
-}
diff --git a/Engine/Source/Editor/MovieSceneTools/Private/TrackEditors/PropertyTrackEditors/VectorPropertyTrackEditor.h b/Engine/Source/Editor/MovieSceneTools/Private/TrackEditors/PropertyTrackEditors/VectorPropertyTrackEditor.h
--- a/Engine/Source/Editor/MovieSceneTools/Private/TrackEditors/PropertyTrackEditors/VectorPropertyTrackEditor.h
+++ b/Engine/Source/Editor/MovieSceneTools/Private/TrackEditors/PropertyTrackEditors/VectorPropertyTrackEditor.h
@@ -35,12 +35,9 @@ protected:
 
 	// FPropertyTrackEditor interface
 	virtual TSharedRef<FPropertySection> MakePropertySectionInterface( UMovieSceneSection& SectionObject, UMovieSceneTrack& Track ) override;
-	virtual void GenerateKeysFromPropertyChanged( const FPropertyChangedParams& PropertyChangedParams, TArray<FVectorKey>& GeneratedKeys ) override;
+	virtual void GenerateKeysFromPropertyChanged( const FPropertyChangedParams& PropertyChangedParams, TArray<FVectorKey>& NewGeneratedKeys, TArray<FVectorKey>& DefaultGeneratedKeys ) override;
 	virtual void InitializeNewTrack( UMovieSceneVectorTrack* NewTrack, FPropertyChangedParams PropertyChangedParams ) override;
 
-	// FKeyframeTrackEditor interface
-	virtual bool ShouldAddKey(UMovieSceneVectorTrack* InTrack, FVectorKey InKey, FKeyParams InKeyParams) const override;
-
 private:
 	static FName XName;
 	static FName YName;
diff --git a/Engine/Source/Editor/MovieSceneTools/Private/TrackEditors/PropertyTrackEditors/VisibilityPropertyTrackEditor.cpp b/Engine/Source/Editor/MovieSceneTools/Private/TrackEditors/PropertyTrackEditors/VisibilityPropertyTrackEditor.cpp
--- a/Engine/Source/Editor/MovieSceneTools/Private/TrackEditors/PropertyTrackEditors/VisibilityPropertyTrackEditor.cpp
+++ b/Engine/Source/Editor/MovieSceneTools/Private/TrackEditors/PropertyTrackEditors/VisibilityPropertyTrackEditor.cpp
@@ -18,7 +18,7 @@ TSharedRef<FPropertySection> FVisibilityPropertyTrackEditor::MakePropertySection
 }
 
 
-void FVisibilityPropertyTrackEditor::GenerateKeysFromPropertyChanged( const FPropertyChangedParams& PropertyChangedParams, TArray<bool>& GeneratedKeys )
+void FVisibilityPropertyTrackEditor::GenerateKeysFromPropertyChanged( const FPropertyChangedParams& PropertyChangedParams, TArray<bool>& NewGeneratedKeys, TArray<bool>& DefaultGeneratedKeys )
 {
-	GeneratedKeys.Add(!PropertyChangedParams.GetPropertyValue<bool>());
+	NewGeneratedKeys.Add(!PropertyChangedParams.GetPropertyValue<bool>());
 }
diff --git a/Engine/Source/Editor/MovieSceneTools/Private/TrackEditors/PropertyTrackEditors/VisibilityPropertyTrackEditor.h b/Engine/Source/Editor/MovieSceneTools/Private/TrackEditors/PropertyTrackEditors/VisibilityPropertyTrackEditor.h
--- a/Engine/Source/Editor/MovieSceneTools/Private/TrackEditors/PropertyTrackEditors/VisibilityPropertyTrackEditor.h
+++ b/Engine/Source/Editor/MovieSceneTools/Private/TrackEditors/PropertyTrackEditors/VisibilityPropertyTrackEditor.h
@@ -33,6 +33,6 @@ protected:
 
 	// FPropertyTrackEditor interface
 	virtual TSharedRef<FPropertySection> MakePropertySectionInterface( UMovieSceneSection& SectionObject, UMovieSceneTrack& Track ) override;
-	virtual void GenerateKeysFromPropertyChanged( const FPropertyChangedParams& PropertyChangedParams, TArray<bool>& GeneratedKeys ) override;
+	virtual void GenerateKeysFromPropertyChanged( const FPropertyChangedParams& PropertyChangedParams, TArray<bool>& NewGeneratedKeys, TArray<bool>& DefaultGeneratedKeys ) override;
 	virtual bool ForCustomizedUseOnly() { return true; }
 };
diff --git a/Engine/Source/Editor/MovieSceneTools/Private/TrackEditors/ShotTrackEditor/ShotTrackEditor.cpp b/Engine/Source/Editor/MovieSceneTools/Private/TrackEditors/ShotTrackEditor/ShotTrackEditor.cpp
--- a/Engine/Source/Editor/MovieSceneTools/Private/TrackEditors/ShotTrackEditor/ShotTrackEditor.cpp
+++ b/Engine/Source/Editor/MovieSceneTools/Private/TrackEditors/ShotTrackEditor/ShotTrackEditor.cpp
@@ -44,21 +44,6 @@ TSharedRef<ISequencerTrackEditor> FShotTrackEditor::CreateTrackEditor(TSharedRef
 /* ISequencerTrackEditor interface
  *****************************************************************************/
 
-void FShotTrackEditor::AddKey(const FGuid& ObjectGuid, UObject* AdditionalAsset)
-{/*
-	TArray<UObject*> OutObjects;
-
-	// @todo sequencer: the director track may be able to get cameras from sub-movie scenes
-	GetSequencer()->GetRuntimeObjects(GetSequencer()->GetRootMovieSceneSequenceInstance(), ObjectGuid, OutObjects);
-	bool bValidKey = OutObjects.Num() == 1 && OutObjects[0]->IsA<ACameraActor>();
-	
-	if (bValidKey)*/
-	{
-		AnimatablePropertyChanged(FOnKeyProperty::CreateRaw(this, &FShotTrackEditor::AddKeyInternal, ObjectGuid));
-	}
-}
-
-
 void FShotTrackEditor::BuildAddTrackMenu(FMenuBuilder& MenuBuilder)
 {
 	UMovieSceneSequence* RootMovieSceneSequence = GetSequencer()->GetRootMovieSceneSequence();
@@ -371,7 +356,7 @@ void FShotTrackEditor::HandleAddShotComboButtonMenuEntryExecute(AActor* Camera)
 
 	if (ObjectGuid.IsValid())
 	{
-		AddKey(ObjectGuid);
+		AnimatablePropertyChanged(FOnKeyProperty::CreateRaw(this, &FShotTrackEditor::AddKeyInternal, ObjectGuid));
 	}
 }
 
diff --git a/Engine/Source/Editor/MovieSceneTools/Private/TrackEditors/ShotTrackEditor/ShotTrackEditor.h b/Engine/Source/Editor/MovieSceneTools/Private/TrackEditors/ShotTrackEditor/ShotTrackEditor.h
--- a/Engine/Source/Editor/MovieSceneTools/Private/TrackEditors/ShotTrackEditor/ShotTrackEditor.h
+++ b/Engine/Source/Editor/MovieSceneTools/Private/TrackEditors/ShotTrackEditor/ShotTrackEditor.h
@@ -38,7 +38,6 @@ public:
 
 	// ISequencerTrackEditor interface
 
-	virtual void AddKey(const FGuid& ObjectGuid, UObject* AdditionalAsset = nullptr) override;
 	virtual void BuildAddTrackMenu(FMenuBuilder& MenuBuilder) override;
 	virtual TSharedPtr<SWidget> BuildOutlinerEditWidget(const FGuid& ObjectBinding, UMovieSceneTrack* Track) override;
 	virtual TSharedRef<ISequencerSection> MakeSectionInterface(UMovieSceneSection& SectionObject, UMovieSceneTrack& Track) override;
diff --git a/Engine/Source/Editor/MovieSceneTools/Private/TrackEditors/SkeletalAnimationTrackEditor.cpp b/Engine/Source/Editor/MovieSceneTools/Private/TrackEditors/SkeletalAnimationTrackEditor.cpp
--- a/Engine/Source/Editor/MovieSceneTools/Private/TrackEditors/SkeletalAnimationTrackEditor.cpp
+++ b/Engine/Source/Editor/MovieSceneTools/Private/TrackEditors/SkeletalAnimationTrackEditor.cpp
@@ -146,16 +146,36 @@ TSharedRef<ISequencerSection> FSkeletalAnimationTrackEditor::MakeSectionInterfac
 }
 
 
-void FSkeletalAnimationTrackEditor::AddKey(const FGuid& ObjectGuid, UObject* AdditionalAsset)
+void FSkeletalAnimationTrackEditor::AddKey(const FGuid& ObjectGuid)
 {
-	UAnimSequence* AnimSequence = Cast<UAnimSequence>(AdditionalAsset);
+	USkeleton* Skeleton = AcquireSkeletonFromObjectGuid(ObjectGuid);
 
-	if (AnimSequence)
+	if (Skeleton)
 	{
-		TArray<UObject*> OutObjects;
+		// Load the asset registry module
+		FAssetRegistryModule& AssetRegistryModule = FModuleManager::LoadModuleChecked<FAssetRegistryModule>(TEXT("AssetRegistry"));
 
-		GetSequencer()->GetRuntimeObjects( GetSequencer()->GetFocusedMovieSceneSequenceInstance(), ObjectGuid, OutObjects);
-		AnimatablePropertyChanged( FOnKeyProperty::CreateRaw( this, &FSkeletalAnimationTrackEditor::AddKeyInternal, OutObjects, AnimSequence) );
+		// Collect a full list of assets with the specified class
+		TArray<FAssetData> AssetDataList;
+		AssetRegistryModule.Get().GetAssetsByClass(UAnimSequence::StaticClass()->GetFName(), AssetDataList);
+
+		if (AssetDataList.Num())
+		{
+			FMenuBuilder MenuBuilder(true, NULL);
+			BuildAnimationSubMenu(MenuBuilder, ObjectGuid, Skeleton);
+			
+			TSharedPtr< SWindow > Parent = FSlateApplication::Get().GetActiveTopLevelWindow(); 
+			if (Parent.IsValid())
+			{
+				FSlateApplication::Get().PushMenu(
+					Parent.ToSharedRef(),
+					FWidgetPath(),
+					MenuBuilder.MakeWidget(),
+					FSlateApplication::Get().GetCursorPos(),
+					FPopupTransitionEffect(FPopupTransitionEffect::TypeInPopup)
+					);
+			}
+		}
 	}
 }
 
@@ -240,12 +260,18 @@ void FSkeletalAnimationTrackEditor::BuildAnimationSubMenu(FMenuBuilder& MenuBuil
 
 void FSkeletalAnimationTrackEditor::OnAnimationAssetSelected(const FAssetData& AssetData, FGuid ObjectBinding)
 {
+	FSlateApplication::Get().DismissAllMenus();
+
 	UObject* SelectedObject = AssetData.GetAsset();
 
 	if (SelectedObject && SelectedObject->IsA(UAnimSequence::StaticClass()))
 	{
 		UAnimSequence* AnimSequence = CastChecked<UAnimSequence>(AssetData.GetAsset());
-		AddKey(ObjectBinding, AnimSequence);
+
+		TArray<UObject*> OutObjects;
+
+		GetSequencer()->GetRuntimeObjects( GetSequencer()->GetFocusedMovieSceneSequenceInstance(), ObjectBinding, OutObjects);
+		AnimatablePropertyChanged( FOnKeyProperty::CreateRaw( this, &FSkeletalAnimationTrackEditor::AddKeyInternal, OutObjects, AnimSequence) );
 	}
 }
 
diff --git a/Engine/Source/Editor/MovieSceneTools/Private/TrackEditors/SkeletalAnimationTrackEditor.h b/Engine/Source/Editor/MovieSceneTools/Private/TrackEditors/SkeletalAnimationTrackEditor.h
--- a/Engine/Source/Editor/MovieSceneTools/Private/TrackEditors/SkeletalAnimationTrackEditor.h
+++ b/Engine/Source/Editor/MovieSceneTools/Private/TrackEditors/SkeletalAnimationTrackEditor.h
@@ -31,7 +31,7 @@ public:
 
 	// ISequencerTrackEditor interface
 
-	virtual void AddKey(const FGuid& ObjectGuid, UObject* AdditionalAsset) override;
+	virtual void AddKey(const FGuid& ObjectGuid) override;
 	virtual void BuildObjectBindingTrackMenu(FMenuBuilder& MenuBuilder, const FGuid& ObjectBinding, const UClass* ObjectClass) override;
 	virtual bool HandleAssetAdded(UObject* Asset, const FGuid& TargetObjectGuid) override;
 	virtual TSharedRef<ISequencerSection> MakeSectionInterface( UMovieSceneSection& SectionObject, UMovieSceneTrack& Track ) override;
diff --git a/Engine/Source/Editor/MovieSceneTools/Private/TrackEditors/TransformTrackEditor.cpp b/Engine/Source/Editor/MovieSceneTools/Private/TrackEditors/TransformTrackEditor.cpp
--- a/Engine/Source/Editor/MovieSceneTools/Private/TrackEditors/TransformTrackEditor.cpp
+++ b/Engine/Source/Editor/MovieSceneTools/Private/TrackEditors/TransformTrackEditor.cpp
@@ -347,12 +347,19 @@ void F3DTransformTrackEditor::OnTransformChanged( UObject& InObject )
 		FTransformData NewTransformData( SceneComponentThatChanged );
 
 		FKeyParams KeyParams;
+		KeyParams.bCreateKeyIfUnchanged = GetSequencer()->GetKeyAllEnabled();
+		KeyParams.bCreateKeyOnlyWhenAutoKeying = true;
+		if ( GetSequencer()->GetAutoKeyMode() == EAutoKeyMode::KeyAll )
+		{
+			KeyParams.bCreateHandleIfMissing = true;
+			KeyParams.bCreateTrackIfMissing = true;
+			KeyParams.bCreateKeyIfEmpty = true;
+		}
+		else
 		{
 			KeyParams.bCreateHandleIfMissing = false;
 			KeyParams.bCreateTrackIfMissing = false;
-			KeyParams.bCreateKeyIfUnchanged = GetSequencer()->GetKeyAllEnabled();
 			KeyParams.bCreateKeyIfEmpty = false;
-			KeyParams.bCreateKeyOnlyWhenAutoKeying = true;
 		}
 
 		const bool bUnwindRotation = GetSequencer()->IsRecordingLive();
@@ -362,15 +369,6 @@ void F3DTransformTrackEditor::OnTransformChanged( UObject& InObject )
 }
 
 
-void F3DTransformTrackEditor::AddKey(const FGuid& ObjectGuid, UObject* AdditionalAsset)
-{
-	// @todo - Sequencer - The key params should be updated to do what the artists expect.
-	FKeyParams KeyParams;
-	KeyParams.bCreateTrackIfMissing = true;
-	AddTransformKeysForHandle( ObjectGuid, EKey3DTransformChannel::All, KeyParams);
-}
-
-
 void F3DTransformTrackEditor::OnAddTransformKeysForSelectedObjects( EKey3DTransformChannel::Type Channel )
 {
 	// WASD hotkeys to fly the viewport can conflict with hotkeys for setting keyframes (ie. s). 
@@ -643,23 +641,21 @@ void F3DTransformTrackEditor::GetActorAndSceneComponentFromObject( UObject* Obje
 	}
 }
 
-void GetKeysForVector( bool LastVectorIsValid, const FVector& LastVector, const FVector& CurrentVector, EKey3DTransformChannel::Type KeyChannel, EKey3DTransformChannel::ValueType KeyChannelValueType, bool bUnwindRotation, TArray<FTransformKey>& OutKeys )
+void GetKeysForVector( bool LastVectorIsValid, const FVector& LastVector, const FVector& CurrentVector, EKey3DTransformChannel::Type VectorChannel, EKey3DTransformChannel::Type ChannelsToKey, bool bUnwindRotation, TArray<FTransformKey>& OutNewKeys, TArray<FTransformKey>& OutDefaultKeys)
 {
-	if ( LastVectorIsValid == false || FMath::IsNearlyEqual( LastVector.X, CurrentVector.X ) == false )
-	{
-		OutKeys.Add( FTransformKey( KeyChannel, KeyChannelValueType, EAxis::X, CurrentVector.X, bUnwindRotation ) );
-	}
-	if ( LastVectorIsValid == false || FMath::IsNearlyEqual( LastVector.Y, CurrentVector.Y ) == false )
-	{
-		OutKeys.Add( FTransformKey( KeyChannel, KeyChannelValueType, EAxis::Y, CurrentVector.Y, bUnwindRotation ) );
-	}
-	if ( LastVectorIsValid == false || FMath::IsNearlyEqual( LastVector.Z, CurrentVector.Z ) == false )
-	{
-		OutKeys.Add( FTransformKey( KeyChannel, KeyChannelValueType, EAxis::Z, CurrentVector.Z, bUnwindRotation ) );
-	}
+	bool bKeyChannel = ChannelsToKey == EKey3DTransformChannel::All || ChannelsToKey == VectorChannel;
+
+	TArray<FTransformKey>& XKeys = bKeyChannel && ( LastVectorIsValid == false || FMath::IsNearlyEqual( LastVector.X, CurrentVector.X ) == false ) ? OutNewKeys : OutDefaultKeys;
+	XKeys.Add( FTransformKey( VectorChannel, EAxis::X, CurrentVector.X, bUnwindRotation ) );
+	
+	TArray<FTransformKey>& YKeys = bKeyChannel && ( LastVectorIsValid == false || FMath::IsNearlyEqual( LastVector.Y, CurrentVector.Y ) == false ) ? OutNewKeys : OutDefaultKeys;
+	YKeys.Add( FTransformKey( VectorChannel, EAxis::Y, CurrentVector.Y, bUnwindRotation ) );
+
+	TArray<FTransformKey>& ZKeys = bKeyChannel && ( LastVectorIsValid == false || FMath::IsNearlyEqual( LastVector.Z, CurrentVector.Z ) == false ) ? OutNewKeys : OutDefaultKeys;
+	ZKeys.Add( FTransformKey( VectorChannel, EAxis::Z, CurrentVector.Z, bUnwindRotation ) );
 }
 
-void F3DTransformTrackEditor::GetChangedTransformKeys( const FTransformData& LastTransform, const FTransformData& CurrentTransform, EKey3DTransformChannel::Type ChannelsToKey, bool bUnwindRotation, TArray<FTransformKey>& OutKeys )
+void F3DTransformTrackEditor::GetTransformKeys( const FTransformData& LastTransform, const FTransformData& CurrentTransform, EKey3DTransformChannel::Type ChannelsToKey, bool bUnwindRotation, TArray<FTransformKey>& OutNewKeys, TArray<FTransformKey>& OutDefaultKeys )
 {
 	bool bLastVectorIsValid = LastTransform.IsValid();
 
@@ -670,14 +666,9 @@ void F3DTransformTrackEditor::GetChangedTransformKeys( const FTransformData& Las
 		ChannelsToKey = EKey3DTransformChannel::All;
 	}
 
-	GetKeysForVector( bLastVectorIsValid, LastTransform.Translation, CurrentTransform.Translation, EKey3DTransformChannel::Translation, ChannelsToKey & EKey3DTransformChannel::Translation ? EKey3DTransformChannel::Key : EKey3DTransformChannel::Default, bUnwindRotation, OutKeys);
-	GetKeysForVector( bLastVectorIsValid, LastTransform.Rotation.Euler(), CurrentTransform.Rotation.Euler(), EKey3DTransformChannel::Rotation, ChannelsToKey & EKey3DTransformChannel::Rotation ? EKey3DTransformChannel::Key : EKey3DTransformChannel::Default, bUnwindRotation, OutKeys );
-	GetKeysForVector( bLastVectorIsValid, LastTransform.Scale, CurrentTransform.Scale, EKey3DTransformChannel::Scale, ChannelsToKey & EKey3DTransformChannel::Scale ? EKey3DTransformChannel::Key : EKey3DTransformChannel::Default, bUnwindRotation, OutKeys );
-}
-
-bool F3DTransformTrackEditor::ShouldAddKey(UMovieScene3DTransformTrack* InTrack, FTransformKey InKey, FKeyParams InKeyParams) const
-{
-	return FKeyframeTrackEditor::ShouldAddKey(InTrack, InKey, InKeyParams) && InKey.ChannelValueType == EKey3DTransformChannel::Key;
+	GetKeysForVector( bLastVectorIsValid, LastTransform.Translation, CurrentTransform.Translation, EKey3DTransformChannel::Translation, ChannelsToKey, bUnwindRotation, OutNewKeys, OutDefaultKeys );
+	GetKeysForVector( bLastVectorIsValid, LastTransform.Rotation.Euler(), CurrentTransform.Rotation.Euler(), EKey3DTransformChannel::Rotation, ChannelsToKey, bUnwindRotation, OutNewKeys, OutDefaultKeys );
+	GetKeysForVector( bLastVectorIsValid, LastTransform.Scale, CurrentTransform.Scale, EKey3DTransformChannel::Scale, ChannelsToKey, bUnwindRotation, OutNewKeys, OutDefaultKeys );
 }
 
 void F3DTransformTrackEditor::AddTransformKeysForHandle( FGuid ObjectHandle, EKey3DTransformChannel::Type ChannelToKey, FKeyParams KeyParams )
@@ -707,14 +698,15 @@ void F3DTransformTrackEditor::AddTransformKeysForObject( UObject* Object, EKey3D
 
 void F3DTransformTrackEditor::AddTransformKeys( AActor* ActorToKey, const FTransformData& LastTransform, const FTransformData& CurrentTransform, EKey3DTransformChannel::Type ChannelsToKey, bool bUnwindRotation, FKeyParams KeyParams )
 {
-	TArray<FTransformKey> Keys;
-	GetChangedTransformKeys(LastTransform, CurrentTransform, ChannelsToKey, bUnwindRotation, Keys);
+	TArray<FTransformKey> NewKeys;
+	TArray<FTransformKey> DefaultKeys;
+	GetTransformKeys(LastTransform, CurrentTransform, ChannelsToKey, bUnwindRotation, NewKeys, DefaultKeys);
 
-	AnimatablePropertyChanged( FOnKeyProperty::CreateRaw(this, &F3DTransformTrackEditor::OnAddTransformKeys, ActorToKey, &Keys, CurrentTransform, KeyParams ) );
+	AnimatablePropertyChanged( FOnKeyProperty::CreateRaw(this, &F3DTransformTrackEditor::OnAddTransformKeys, ActorToKey, &NewKeys, &DefaultKeys, CurrentTransform, KeyParams ) );
 }
 
 
-bool F3DTransformTrackEditor::OnAddTransformKeys( float Time, AActor* ActorToKey, TArray<FTransformKey>* Keys, FTransformData CurrentTransform, FKeyParams KeyParams )
+bool F3DTransformTrackEditor::OnAddTransformKeys( float Time, AActor* ActorToKey, TArray<FTransformKey>* NewKeys, TArray<FTransformKey>* DefaultKeys, FTransformData CurrentTransform, FKeyParams KeyParams )
 {
 	TArray<UObject*> ObjectsToKey;
 	ObjectsToKey.Add(ActorToKey);
@@ -722,7 +714,8 @@ bool F3DTransformTrackEditor::OnAddTransformKeys( float Time, AActor* ActorToKey
 	return AddKeysToObjects(
 		ObjectsToKey,
 		Time,
-		*Keys,
+		*NewKeys,
+		*DefaultKeys,
 		KeyParams,
 		UMovieScene3DTransformTrack::StaticClass(),
 		TransformPropertyName,
diff --git a/Engine/Source/Editor/MovieSceneTools/Private/TrackEditors/TransformTrackEditor.h b/Engine/Source/Editor/MovieSceneTools/Private/TrackEditors/TransformTrackEditor.h
--- a/Engine/Source/Editor/MovieSceneTools/Private/TrackEditors/TransformTrackEditor.h
+++ b/Engine/Source/Editor/MovieSceneTools/Private/TrackEditors/TransformTrackEditor.h
@@ -37,7 +37,6 @@ public:
 
 	// ISequencerTrackEditor interface
 
-	virtual void AddKey( const FGuid& ObjectGuid, UObject* AdditionalAsset = NULL ) override;
 	virtual void BindCommands(TSharedRef<FUICommandList> SequencerCommandBindings) override;
 	virtual void BuildObjectBindingEditButtons(TSharedPtr<SHorizontalBox> EditBox, const FGuid& ObjectBinding, const UClass* ObjectClass) override;
 	virtual void BuildObjectBindingTrackMenu(FMenuBuilder& MenuBuilder, const FGuid& ObjectBinding, const UClass* ObjectClass) override;
@@ -46,9 +45,6 @@ public:
 	virtual bool SupportsType( TSubclassOf<UMovieSceneTrack> Type ) const override;
 	virtual void BuildTrackContextMenu( FMenuBuilder& MenuBuilder, UMovieSceneTrack* Track ) override;
 
-	// FKeyframeTrackEditor interface
-	virtual bool ShouldAddKey(UMovieScene3DTransformTrack* InTrack, FTransformKey InKey, FKeyParams InKeyParams) const override;
-
 private:
 
 	/** Returns whether or not a transform track can be added for an actor with a specific handle. */
@@ -85,7 +81,7 @@ private:
 
 	/** Generates transform keys based on the last transform, the current transform, and other options. 
 		One transform key is generated for each individual key to be added to the section. */
-	void GetChangedTransformKeys(const FTransformData& LastTransform, const FTransformData& CurrentTransform, EKey3DTransformChannel::Type ChannelsToKey, bool bUnwindRotation, TArray<FTransformKey>& OutKeys );
+	void GetTransformKeys( const FTransformData& LastTransform, const FTransformData& CurrentTransform, EKey3DTransformChannel::Type ChannelsToKey, bool bUnwindRotation, TArray<FTransformKey>& OutNewKeys, TArray<FTransformKey>& OutDefaultKeys );
 
 	/**
 	* Adds transform tracks and keys to the selected objects in the level.
@@ -131,7 +127,7 @@ private:
 	* @param Keys The keys to add.
 	* @param KeyParams Parameters which control how the keys are added.
 	*/
-	bool OnAddTransformKeys( float Time, AActor* ActorToKey, TArray<FTransformKey>* Keys, FTransformData CurrentTransform, FKeyParams KeyParams );
+	bool OnAddTransformKeys( float Time, AActor* ActorToKey, TArray<FTransformKey>* NewKeys, TArray<FTransformKey>* DefaultKeys, FTransformData CurrentTransform, FKeyParams KeyParams );
 
 private:
 
diff --git a/Engine/Source/Editor/MovieSceneTools/Public/KeyframeTrackEditor.h b/Engine/Source/Editor/MovieSceneTools/Public/KeyframeTrackEditor.h
--- a/Engine/Source/Editor/MovieSceneTools/Public/KeyframeTrackEditor.h
+++ b/Engine/Source/Editor/MovieSceneTools/Public/KeyframeTrackEditor.h
@@ -39,12 +39,6 @@ public:
 		return Type == TrackType::StaticClass();
 	}
 
-	// Only add a key if the track already has keys, or if creating keys has been forced.  Otherwise just update the track default.
-	virtual bool ShouldAddKey(TrackType* InTrack, KeyDataType InNewKey, FKeyParams InKeyParams) const
-	{
-		return HasKeys( InTrack, InNewKey ) || InKeyParams.bCreateKeyIfEmpty;
-	}
-
 protected:
 
 	/*
@@ -53,6 +47,8 @@ protected:
 	 * @param ObjectsToKey An array of objects to add keyframes to.
 	 * @param KeyTime The time to add keys.
 	 * @param NewKeys The new keys to add.
+	 * @param DefaultKeys Extra keys with default values which shouldn't be added directly, but are needed to set correct
+	 *        default values when adding single channel keys for multi-channel tracks like vectors.
 	 * @param KeyParams The parameters to control keyframing behavior.
 	 * @param TrackClass The class of track which should be created if specified in the parameters.
 	 * @param PropertyName The name of the property to add keys for.
@@ -60,15 +56,17 @@ protected:
 	 *        track is created, but before any sections or keys have been added.
 	 * @param OnSetIntermediateValue A delegate which is called when the key params specify that keying should only happen when
 	 * 		  auto-keying is on, but auto-keying is off.  This allows handling the changed but not keyed intermediate value.
+	 * @return Whether or not a handle guid or track was created. Note this does not return true if keys were added or modified.
 	 */
 	bool AddKeysToObjects(
-		TArray<UObject*> ObjectsToKey, float KeyTime, const TArray<KeyDataType> NewKeys, FKeyParams KeyParams,
-		TSubclassOf<UMovieSceneTrack> TrackClass, FName PropertyName,
+		TArray<UObject*> ObjectsToKey, float KeyTime,
+		const TArray<KeyDataType>& NewKeys, const TArray<KeyDataType>& DefaultKeys,
+		FKeyParams KeyParams, TSubclassOf<UMovieSceneTrack> TrackClass, FName PropertyName,
 		TFunction<void(TrackType*)> OnIntializeNewTrack,
 		TFunction<void(TrackType*)> OnSetIntermediateValue)
 	{
 		bool bHandleCreated = false;
-		bool bTrackCreatedOrModified = false;
+		bool bTrackCreated = false;
 
 		for ( UObject* Object : ObjectsToKey )
 		{
@@ -78,23 +76,40 @@ protected:
 
 			if ( ObjectHandle.IsValid() )
 			{
-				bTrackCreatedOrModified |= AddKeysToHandle( ObjectHandle, KeyTime, NewKeys, KeyParams, TrackClass, PropertyName, OnIntializeNewTrack, OnSetIntermediateValue );
+				bTrackCreated |= AddKeysToHandle( ObjectHandle, KeyTime, NewKeys, DefaultKeys, KeyParams, TrackClass, PropertyName, OnIntializeNewTrack, OnSetIntermediateValue );
 			}
 		}
-		return bHandleCreated || bTrackCreatedOrModified;
+		return bHandleCreated || bTrackCreated;
 	}
 
 
 private:
 
+	/*
+	 * Adds keys to the specified guid.  This may also add tracks and sections depending on the options specified. 
+	 *
+	 * @param ObjectsToKey An array of objects to add keyframes to.
+	 * @param KeyTime The time to add keys.
+	 * @param NewKeys The new keys to add.
+	 * @param DefaultKeys Extra keys with default values which shouldn't be added directly, but are needed to set correct
+	 *        default values when adding single channel keys for multi-channel tracks like vectors.
+	 * @param KeyParams The parameters to control keyframing behavior.
+	 * @param TrackClass The class of track which should be created if specified in the parameters.
+	 * @param PropertyName The name of the property to add keys for.
+	 * @param OnInitializeNewTrack A delegate which allows for custom initialization for new tracks.  This is called after the 
+	 *        track is created, but before any sections or keys have been added.
+	 * @param OnSetIntermediateValue A delegate which is called when the key params specify that keying should only happen when
+	 * 		  auto-keying is on, but auto-keying is off.  This allows handling the changed but not keyed intermediate value.
+	 * @return Whether or not a track was created. Note this does not return true if keys were added or modified.
+	*/
 	bool AddKeysToHandle(
-		FGuid ObjectHandle, float KeyTime, const TArray<KeyDataType> NewKeys, FKeyParams KeyParams,
-		TSubclassOf<UMovieSceneTrack> TrackClass, FName PropertyName,
+		FGuid ObjectHandle, float KeyTime,
+		const TArray<KeyDataType>& NewKeys, const TArray<KeyDataType>& DefaultKeys,
+		FKeyParams KeyParams, TSubclassOf<UMovieSceneTrack> TrackClass, FName PropertyName,
 		TFunction<void(TrackType*)> OnIntializeNewTrack,
 		TFunction<void(TrackType*)> OnSetIntermediateValue)
 	{
 		bool bTrackCreated = false;
-		bool bTrackModified = false;
 
 		// Try to find an existing Track, and if one doesn't exist check the key params and create one if requested.
 		FFindOrCreateTrackResult TrackResult = FindOrCreateTrackForObject( ObjectHandle, TrackClass, PropertyName, KeyParams.bCreateTrackIfMissing );
@@ -111,19 +126,19 @@ private:
 
 		if ( Track )
 		{
-			bTrackModified |= AddKeysToTrack( Track, KeyTime, NewKeys, KeyParams, OnSetIntermediateValue );
+			AddKeysToTrack( Track, KeyTime, NewKeys, DefaultKeys, KeyParams, bTrackCreated, OnSetIntermediateValue );
 		}
 
-		return bTrackCreated || bTrackModified;
+		return bTrackCreated;
 	}
 
-	bool AddKeysToTrack(
-		TrackType* Track, float KeyTime, const TArray<KeyDataType> NewKeys, FKeyParams KeyParams,
+	void AddKeysToTrack(
+		TrackType* Track, float KeyTime,
+		const TArray<KeyDataType>& NewKeys, const TArray<KeyDataType>& DefaultKeys,
+		FKeyParams KeyParams, bool bNewTrack,
 		TFunction<void(TrackType*)> OnSetIntermediateValue)
 	{
-		bool bTrackModified = false;
-
-		bool bSettingIntermediateValue = KeyParams.bCreateKeyOnlyWhenAutoKeying && GetSequencer()->GetAutoKeyEnabled() == false;
+		bool bSettingIntermediateValue = KeyParams.bCreateKeyOnlyWhenAutoKeying && GetSequencer()->GetAutoKeyMode() == EAutoKeyMode::KeyNone;
 		if ( bSettingIntermediateValue )
 		{
 			checkf(OnSetIntermediateValue, TEXT("A valid OnSetIntermediateValue delegate must be provided for consistent keyframing behavior."));
@@ -137,7 +152,7 @@ private:
 				// Only modify the track if the new key will create new data, or if creating keys has been forced.
 				if ( NewKeyIsNewData( Track, KeyTime, NewKey ) || KeyParams.bCreateKeyIfUnchanged )
 				{
-					if ( ShouldAddKey(Track, NewKey, KeyParams) )
+					if ( HasKeys( Track, NewKey ) || KeyParams.bCreateKeyIfEmpty )
 					{
 						AddKey( Track, KeyTime, NewKey, InterpolationMode );
 					}
@@ -145,11 +160,17 @@ private:
 					{
 						SetDefault( Track, KeyTime, NewKey );
 					}
-					bTrackModified = true;
+				}
+			}
+			// If a new track was added, set defaults for all of the default keys.
+			if( bNewTrack )
+			{
+				for ( const KeyDataType& DefaultKey : DefaultKeys )
+				{
+					SetDefault( Track, KeyTime, DefaultKey );
 				}
 			}
 		}
-		return bTrackModified;
 	}
 
 	bool NewKeyIsNewData( TrackType* Track, float Time, const KeyDataType& KeyData ) const
diff --git a/Engine/Source/Editor/MovieSceneTools/Public/PropertyTrackEditor.h b/Engine/Source/Editor/MovieSceneTools/Public/PropertyTrackEditor.h
--- a/Engine/Source/Editor/MovieSceneTools/Public/PropertyTrackEditor.h
+++ b/Engine/Source/Editor/MovieSceneTools/Public/PropertyTrackEditor.h
@@ -110,19 +110,6 @@ public:
 		return Type == TrackType::StaticClass();
 	}
 
-	virtual void AddKey( const FGuid& ObjectGuid, UObject* AdditionalAsset = NULL ) override
-	{
-		ISequencerObjectChangeListener& ObjectChangeListener = FMovieSceneTrackEditor::GetSequencer()->GetObjectChangeListener();
-		TSharedPtr<ISequencer> Sequencer = FMovieSceneTrackEditor::GetSequencer();
-
-		TArray<UObject*> OutObjects;
-		Sequencer->GetRuntimeObjects(Sequencer->GetFocusedMovieSceneSequenceInstance(), ObjectGuid, OutObjects);
-		for ( int32 i = 0; i < OutObjects.Num(); ++i )
-		{
-			ObjectChangeListener.TriggerAllPropertiesChanged( OutObjects[i] );
-		}
-	}
-
 	virtual TSharedRef<ISequencerSection> MakeSectionInterface( class UMovieSceneSection& SectionObject, UMovieSceneTrack& Track ) override
 	{
 		TSharedRef<FPropertySection> PropertySection = MakePropertySectionInterface( SectionObject, Track );
@@ -137,12 +124,13 @@ protected:
 	virtual TSharedRef<FPropertySection> MakePropertySectionInterface( class UMovieSceneSection& SectionObject, UMovieSceneTrack& Track ) = 0;
 
 	/**
-	* Generates a new key who's value is the result of the supplied property change.
+	* Generates keys based on the new value from the property property change parameters.
 	*
 	* @param PropertyChangedParams Parameters associated with the property change.
-	* @return The new key.
+	* @param NewGeneratedKeys New keys which should be added due to the property change.
+	* @param DefaultGeneratedKeys Default value keys which should not be added, but may be needed for setting up defaults on new multi-channel tracks.
 	*/
-	virtual void GenerateKeysFromPropertyChanged( const FPropertyChangedParams& PropertyChangedParams, TArray<KeyDataType>& GeneratedKeys ) = 0;
+	virtual void GenerateKeysFromPropertyChanged( const FPropertyChangedParams& PropertyChangedParams, TArray<KeyDataType>& NewGeneratedKeys, TArray<KeyDataType>& DefaultGeneratedKeys ) = 0;
 
 	/** When true, this track editor will only be used on properties which have specified it as a custom track class. This is necessary to prevent duplicate
 		property change handling in cases where a custom track editor handles the same type of data as one of the standard track editors. */
@@ -197,8 +185,9 @@ private:
 	/** Adds a key based on a property change. */
 	bool OnKeyProperty( float KeyTime, FPropertyChangedParams PropertyChangedParams )
 	{
-		TArray<KeyDataType> KeysForPropertyChange;
-		GenerateKeysFromPropertyChanged( PropertyChangedParams, KeysForPropertyChange );
+		TArray<KeyDataType> NewKeysForPropertyChange;
+		TArray<KeyDataType> DefaultKeysForPropertyChange;
+		GenerateKeysFromPropertyChanged( PropertyChangedParams, NewKeysForPropertyChange, DefaultKeysForPropertyChange );
 
 		TSubclassOf<UMovieSceneTrack> CustomizedClass = GetCustomizedTrackClass( PropertyChangedParams.PropertyPath.Last() );
 		TSubclassOf<UMovieSceneTrack> TrackClass;
@@ -212,6 +201,14 @@ private:
 			TrackClass = TrackType::StaticClass();
 		}
 
+		if ( FMovieSceneTrackEditor::GetSequencer()->GetAutoKeyMode() == EAutoKeyMode::KeyAll )
+		{
+			// TODO: Setting these here is a bit sketchy but necessary right now without a larger refactor.
+			PropertyChangedParams.KeyParams.bCreateHandleIfMissing = true;
+			PropertyChangedParams.KeyParams.bCreateTrackIfMissing = true;
+			PropertyChangedParams.KeyParams.bCreateKeyIfEmpty = true;
+		}
+
 		// If the track class has been customized for this property then it's possible this track editor doesn't support it, 
 		// also check for track editors which should only be used for customization.
 		if ( SupportsType( TrackClass ) && ( ForCustomizedUseOnly() == false || CustomizedClass != nullptr) )
@@ -219,7 +216,8 @@ private:
 			return FKeyframeTrackEditor<TrackType, SectionType, KeyDataType>::AddKeysToObjects(
 				PropertyChangedParams.ObjectsThatChanged,
 				KeyTime,
-				KeysForPropertyChange,
+				NewKeysForPropertyChange,
+				DefaultKeysForPropertyChange,
 				PropertyChangedParams.KeyParams,
 				TrackClass,
 				PropertyChangedParams.PropertyPath.Last()->GetFName(),
diff --git a/Engine/Source/Editor/NiagaraEditor/Private/NiagaraSequencer.h b/Engine/Source/Editor/NiagaraEditor/Private/NiagaraSequencer.h
--- a/Engine/Source/Editor/NiagaraEditor/Private/NiagaraSequencer.h
+++ b/Engine/Source/Editor/NiagaraEditor/Private/NiagaraSequencer.h
@@ -236,7 +236,7 @@ public:
 		return false;
 	}
 
-	virtual TSharedRef<ISequencerSection> MakeSectionInterface(class UMovieSceneSection& SectionObject, UMovieSceneTrack& Track)
+	virtual TSharedRef<ISequencerSection> MakeSectionInterface(UMovieSceneSection& SectionObject, UMovieSceneTrack& Track)
 	{
 		INiagaraSequencerSection *Sec = new INiagaraSequencerSection(SectionObject);
 		return MakeShareable(Sec);
diff --git a/Engine/Source/Editor/Sequencer/Private/DisplayNodes/SequencerDisplayNode.cpp b/Engine/Source/Editor/Sequencer/Private/DisplayNodes/SequencerDisplayNode.cpp
--- a/Engine/Source/Editor/Sequencer/Private/DisplayNodes/SequencerDisplayNode.cpp
+++ b/Engine/Source/Editor/Sequencer/Private/DisplayNodes/SequencerDisplayNode.cpp
@@ -420,7 +420,7 @@ void FSequencerDisplayNode::BuildContextMenu(FMenuBuilder& MenuBuilder)
 	{
 		MenuBuilder.AddMenuEntry(
 			LOCTEXT("ToggleNodeActive", "Active"),
-			LOCTEXT("ToggleNodeActiveTooltip", "Set this node or selected nodes active/inactive"),
+			LOCTEXT("ToggleNodeActiveTooltip", "Set this track or selected tracks active/inactive"),
 			FSlateIcon(),
 			FUIAction(
 				FExecuteAction::CreateSP(&GetSequencer(), &FSequencer::ToggleNodeActive),
@@ -433,7 +433,7 @@ void FSequencerDisplayNode::BuildContextMenu(FMenuBuilder& MenuBuilder)
 
 		MenuBuilder.AddMenuEntry(
 			LOCTEXT("ToggleNodeLock", "Locked"),
-			LOCTEXT("ToggleNodeLockTooltip", "Lock or unlock this node or selected nodes"),
+			LOCTEXT("ToggleNodeLockTooltip", "Lock or unlock this node or selected tracks"),
 			FSlateIcon(),
 			FUIAction(
 				FExecuteAction::CreateSP(&GetSequencer(), &FSequencer::ToggleNodeLocked),
@@ -446,14 +446,14 @@ void FSequencerDisplayNode::BuildContextMenu(FMenuBuilder& MenuBuilder)
 
 		MenuBuilder.AddMenuEntry(
 			LOCTEXT("DeleteNode", "Delete"),
-			LOCTEXT("DeleteNodeTooltip", "Delete this or selected nodes"),
+			LOCTEXT("DeleteNodeTooltip", "Delete this or selected tracks"),
 			FSlateIcon(),
 			FUIAction(FExecuteAction::CreateSP(&GetSequencer(), &FSequencer::DeleteNode, ThisNode))
 		);
 
 		MenuBuilder.AddMenuEntry(
 			LOCTEXT("RenameNode", "Rename"),
-			LOCTEXT("RenameNodeTooltip", "Rename this node"),
+			LOCTEXT("RenameNodeTooltip", "Rename this track"),
 			FSlateIcon(),
 			FUIAction(
 				FExecuteAction::CreateSP(this, &FSequencerDisplayNode::HandleContextMenuRenameNodeExecute, ThisNode),
@@ -490,7 +490,7 @@ void FSequencerDisplayNode::SetExpansionState(bool bInExpanded)
 
 bool FSequencerDisplayNode::IsExpanded() const
 {
-	return ParentTree.HasActiveFilter() ? true : bExpanded;
+	return bExpanded;
 }
 
 
@@ -500,6 +500,12 @@ bool FSequencerDisplayNode::IsHidden() const
 }
 
 
+bool FSequencerDisplayNode::IsHovered() const
+{
+	return TreeNodeWidgetPtr.IsValid() && TreeNodeWidgetPtr.Pin()->IsHovered();
+}
+
+
 TSharedRef<FGroupedKeyArea> FSequencerDisplayNode::GetKeyGrouping(int32 InSectionIndex)
 {
 	if (!KeyGroupings.IsValidIndex(InSectionIndex))
diff --git a/Engine/Source/Editor/Sequencer/Private/DisplayNodes/SequencerDisplayNode.h b/Engine/Source/Editor/Sequencer/Private/DisplayNodes/SequencerDisplayNode.h
--- a/Engine/Source/Editor/Sequencer/Private/DisplayNodes/SequencerDisplayNode.h
+++ b/Engine/Source/Editor/Sequencer/Private/DisplayNodes/SequencerDisplayNode.h
@@ -293,6 +293,12 @@ public:
 	 */
 	bool IsHidden() const;
 
+	/**
+	 * Check whether the node's tree view or track area widgets are hovered by the user's mouse.
+	 *
+	 * @return true if hovered, false otherwise. */
+	bool IsHovered() const;
+
 	/** Initialize this node with expansion states and virtual offsets */
 	void Initialize(float InVirtualTop, float InVirtualBottom);
 
diff --git a/Engine/Source/Editor/Sequencer/Private/DisplayNodes/SequencerObjectBindingNode.cpp b/Engine/Source/Editor/Sequencer/Private/DisplayNodes/SequencerObjectBindingNode.cpp
--- a/Engine/Source/Editor/Sequencer/Private/DisplayNodes/SequencerObjectBindingNode.cpp
+++ b/Engine/Source/Editor/Sequencer/Private/DisplayNodes/SequencerObjectBindingNode.cpp
@@ -277,6 +277,41 @@ const UClass* FSequencerObjectBindingNode::GetClassForObjectBinding()
 	return ObjectClass;
 }
 
+UObject* FSequencerObjectBindingNode::FindRepresentativeObject()
+{
+	FSequencer& Sequencer = GetSequencer();
+	
+	UObject* BoundObject = Sequencer.GetFocusedMovieSceneSequenceInstance()->FindObject(ObjectBinding, Sequencer);
+	if (BoundObject)
+	{
+		return BoundObject;
+	}
+
+	UMovieScene* FocusedMovieScene = Sequencer.GetFocusedMovieSceneSequence()->GetMovieScene();
+
+	FMovieScenePossessable* Possessable = FocusedMovieScene->FindPossessable(ObjectBinding);
+	// If we're a possessable with a parent spawnable and we don't have the object, we look the object up within the default object of the spawnable
+	if (Possessable && Possessable->GetParent().IsValid())
+	{
+		// If we're a spawnable and we don't have the object, use the default object to build up the track menu
+		FMovieSceneSpawnable* ParentSpawnable = FocusedMovieScene->FindSpawnable(Possessable->GetParent());
+		if (ParentSpawnable)
+		{
+			UObject* ParentObject = ParentSpawnable->GetClass()->GetDefaultObject();
+			if (ParentObject)
+			{
+				return Sequencer.GetFocusedMovieSceneSequence()->FindPossessableObject(ObjectBinding, ParentObject);
+			}
+		}
+	}
+	// If we're a spawnable and we don't have the object, use the default object to build up the track menu
+	else if (FMovieSceneSpawnable* Spawnable = FocusedMovieScene->FindSpawnable(ObjectBinding))
+	{
+		return Spawnable->GetClass()->GetDefaultObject();
+	}
+
+	return nullptr;
+}
 
 /* FSequencerObjectBindingNode callbacks
  *****************************************************************************/
@@ -288,7 +323,7 @@ TSharedRef<SWidget> FSequencerObjectBindingNode::HandleAddTrackComboButtonGetMen
 	//@todo need to resolve this between UMG and the level editor sequencer
 	const bool bUseSubMenus = Sequencer.IsLevelEditorSequencer();
 
-	UObject* BoundObject = Sequencer.GetFocusedMovieSceneSequenceInstance()->FindObject(ObjectBinding, Sequencer);
+	UObject* BoundObject = FindRepresentativeObject();
 
 	ISequencerModule& SequencerModule = FModuleManager::GetModuleChecked<ISequencerModule>( "Sequencer" );
 	TSharedRef<FUICommandList> CommandList(new FUICommandList);
@@ -456,7 +491,7 @@ void FSequencerObjectBindingNode::HandleLabelsSubMenuCreate(FMenuBuilder& MenuBu
 void FSequencerObjectBindingNode::HandlePropertyMenuItemExecute(TArray<UProperty*> PropertyPath)
 {
 	FSequencer& Sequencer = GetSequencer();
-	UObject* BoundObject = Sequencer.GetFocusedMovieSceneSequenceInstance()->FindObject(ObjectBinding, Sequencer);
+	UObject* BoundObject = FindRepresentativeObject();
 
 	TArray<UObject*> KeyableBoundObjects;
 	if (BoundObject != nullptr)
diff --git a/Engine/Source/Editor/Sequencer/Private/DisplayNodes/SequencerObjectBindingNode.h b/Engine/Source/Editor/Sequencer/Private/DisplayNodes/SequencerObjectBindingNode.h
--- a/Engine/Source/Editor/Sequencer/Private/DisplayNodes/SequencerObjectBindingNode.h
+++ b/Engine/Source/Editor/Sequencer/Private/DisplayNodes/SequencerObjectBindingNode.h
@@ -58,6 +58,8 @@ protected:
 
 private:
 
+	UObject* FindRepresentativeObject();
+	
 	TSharedRef<SWidget> HandleAddTrackComboButtonGetMenuContent();
 	
 	void HandleAddTrackSubMenuNew(FMenuBuilder& AddTrackMenuBuilder, TArray<TArray<UProperty*>> KeyablePropertyPath);
diff --git a/Engine/Source/Editor/Sequencer/Private/DisplayNodes/SequencerTrackNode.cpp b/Engine/Source/Editor/Sequencer/Private/DisplayNodes/SequencerTrackNode.cpp
--- a/Engine/Source/Editor/Sequencer/Private/DisplayNodes/SequencerTrackNode.cpp
+++ b/Engine/Source/Editor/Sequencer/Private/DisplayNodes/SequencerTrackNode.cpp
@@ -40,6 +40,11 @@ void FSequencerTrackNode::SetSectionAsKeyArea(TSharedRef<IKeyArea>& KeyArea)
 	TopLevelKeyNode->AddKeyArea( KeyArea );
 }
 
+void FSequencerTrackNode::AddKey(const FGuid& ObjectGuid)
+{
+	AssociatedEditor.AddKey(ObjectGuid);
+}
+
 
 int32 FSequencerTrackNode::GetMaxRowIndex() const
 {
diff --git a/Engine/Source/Editor/Sequencer/Private/DisplayNodes/SequencerTrackNode.h b/Engine/Source/Editor/Sequencer/Private/DisplayNodes/SequencerTrackNode.h
--- a/Engine/Source/Editor/Sequencer/Private/DisplayNodes/SequencerTrackNode.h
+++ b/Engine/Source/Editor/Sequencer/Private/DisplayNodes/SequencerTrackNode.h
@@ -47,6 +47,12 @@ public:
 	 */
 	void SetSectionAsKeyArea(TSharedRef<IKeyArea>& KeyArea);
 	
+	/**
+	 * Adds a key to the track
+	 *
+	 */
+	void AddKey(const FGuid& ObjectGuid);
+
 	/**
 	 * @return All sections in this node
 	 */
diff --git a/Engine/Source/Editor/Sequencer/Private/SAnimationOutlinerTreeNode.cpp b/Engine/Source/Editor/Sequencer/Private/SAnimationOutlinerTreeNode.cpp
--- a/Engine/Source/Editor/Sequencer/Private/SAnimationOutlinerTreeNode.cpp
+++ b/Engine/Source/Editor/Sequencer/Private/SAnimationOutlinerTreeNode.cpp
@@ -20,6 +20,7 @@ void SAnimationOutlinerTreeNode::Construct( const FArguments& InArgs, TSharedRef
 {
 	DisplayNode = Node;
 
+	HoveredBrush = FEditorStyle::GetBrush( "Sequencer.AnimationOutliner.SelectionBorderHover" );
 	SelectedBrush = FEditorStyle::GetBrush( "Sequencer.AnimationOutliner.SelectionBorder" );
 	SelectedBrushInactive = FEditorStyle::GetBrush("Sequencer.AnimationOutliner.SelectionBorderInactive");
 	NotSelectedBrush = FEditorStyle::GetBrush( "NoBorder" );
@@ -49,6 +50,7 @@ void SAnimationOutlinerTreeNode::Construct( const FArguments& InArgs, TSharedRef
 		[
 			SNew(SBox)
 			.HeightOverride_Lambda(NodeHeight)
+			.Padding(FMargin(0.0f, 0.0f, 10.0f, 0.0f))
 			[
 				SNew( SHorizontalBox )
 
@@ -435,12 +437,16 @@ FReply SAnimationOutlinerTreeNode::OnMouseButtonUp( const FGeometry& MyGeometry,
 
 const FSlateBrush* SAnimationOutlinerTreeNode::GetNodeBorderImage() const
 {
-	// Display a highlight when the node is selected
 	FSequencer& Sequencer = DisplayNode->GetSequencer();
 	const bool bIsSelected = Sequencer.GetSelection().IsSelected(DisplayNode.ToSharedRef());
 
 	if (!bIsSelected)
 	{
+		if (DisplayNode->IsHovered())
+		{
+			return HoveredBrush;
+		}
+
 		return NotSelectedBrush;
 	}
 
diff --git a/Engine/Source/Editor/Sequencer/Private/SAnimationOutlinerTreeNode.h b/Engine/Source/Editor/Sequencer/Private/SAnimationOutlinerTreeNode.h
--- a/Engine/Source/Editor/Sequencer/Private/SAnimationOutlinerTreeNode.h
+++ b/Engine/Source/Editor/Sequencer/Private/SAnimationOutlinerTreeNode.h
@@ -27,12 +27,15 @@ public:
 	void EnterRenameMode();
 
 	/**
-	 * @return The display node used by this widget                                                              
+	 * @return The display node used by this widget.                                                           
 	 */
-	const TSharedPtr<FSequencerDisplayNode> GetDisplayNode() const { return DisplayNode; }
+	const TSharedPtr<FSequencerDisplayNode> GetDisplayNode() const
+	{
+		return DisplayNode;
+	}
 
 	/**
-	 * Gets a tint to apply buttons on hover
+	 * Gets a tint to apply buttons on hover.
 	 */
 	FLinearColor GetHoverTint() const;
 
@@ -46,12 +49,12 @@ private:
 	FSlateColor GetForegroundBasedOnSelection() const;
 
 	/**
-	 * @return The border image to show in the tree node
+	 * @return The border image to show in the tree node.
 	 */
 	const FSlateBrush* GetNodeBorderImage() const;
 	
 	/**
-	 * @return The expander visibility of this node
+	 * @return The expander visibility of this node.
 	 */
 	EVisibility GetExpanderVisibility() const;
 
@@ -75,32 +78,35 @@ private:
 	/** Handles the add key button being clicked. */
 	FReply OnAddKeyClicked();
 
-	/** Get all descendant nodes from the given root node */
+	/** Get all descendant nodes from the given root node. */
 	void GetAllDescendantNodes(TSharedPtr<FSequencerDisplayNode> RootNode, TArray<TSharedRef<FSequencerDisplayNode> >& AllNodes);
 
-	/** Return the root node given an object node */
+	/** Return the root node given an object node. */
 	TSharedPtr<FSequencerDisplayNode> GetRootNode(TSharedPtr<FSequencerDisplayNode> ObjectNode);
 
-	/** Called when nodes are selected */
+	/** Called when nodes are selected. */
 	void OnSelectionChanged(TArray<TSharedPtr<FSequencerDisplayNode>> AffectedNodes);
 
 private:
 
-	/** Layout node the widget is visualizing */
+	/** Layout node the widget is visualizing. */
 	TSharedPtr<FSequencerDisplayNode> DisplayNode;
 
-	/** Holds the editable text label widge.t */
+	/** Holds the editable text label widget. */
 	TSharedPtr<SEditableLabel> EditableLabel;
 
-	/** Brush to display a border around the widget when it is selected */
+	/** Brush to use if the node is hovered. */
+	const FSlateBrush* HoveredBrush;
+
+	/** Brush to display a border around the widget when it is selected. */
 	const FSlateBrush* SelectedBrush;
 
-	/** Brush to display a border around the widget when it is selected but inactive */
+	/** Brush to display a border around the widget when it is selected but inactive. */
 	const FSlateBrush* SelectedBrushInactive;
 
-	/** Brush to use if the node is not selected */
+	/** Brush to use if the node is not selected. */
 	const FSlateBrush* NotSelectedBrush;
 
-	/** The table row style used for nodes in the tree. This is required as we don't actually use the tree for selection */
+	/** The table row style used for nodes in the tree. This is required as we don't actually use the tree for selection. */
 	const FTableRowStyle* TableRowStyle;
 };
\ No newline at end of file
diff --git a/Engine/Source/Editor/Sequencer/Private/SSequencer.cpp b/Engine/Source/Editor/Sequencer/Private/SSequencer.cpp
--- a/Engine/Source/Editor/Sequencer/Private/SSequencer.cpp
+++ b/Engine/Source/Editor/Sequencer/Private/SSequencer.cpp
@@ -284,7 +284,6 @@ void SSequencer::Construct(const FArguments& InArgs, TSharedRef<FSequencer> InSe
 																			.FillWidth( FillCoefficient_0 )
 																			[
 																				SNew(SSequencerTreeViewBox, InSequencer, SharedThis(this))
-																					.Padding(FMargin(0, 0, 10.f, 0)) // Padding to allow space for the scroll bar
 																					[
 																						TreeView.ToSharedRef()
 																					]
@@ -583,7 +582,19 @@ void SSequencer::HandleKeySelectionChanged()
 
 void SSequencer::HandleLabelBrowserSelectionChanged(FString NewLabel, ESelectInfo::Type SelectInfo)
 {
-	SearchBox->SetText(FText::FromString(FString(TEXT("label:") + NewLabel)));
+	if (SelectInfo == ESelectInfo::Direct)
+	{
+		return;
+	}
+
+	if (NewLabel.IsEmpty())
+	{
+		SearchBox->SetText(FText::GetEmpty());
+	}
+	else
+	{
+		SearchBox->SetText(FText::FromString(FString(TEXT("label:")) + NewLabel));
+	}
 }
 
 
@@ -696,7 +707,38 @@ TSharedRef<SWidget> SSequencer::MakeToolBar()
 			ToolBarBuilder.AddToolBarButton( FSequencerCommands::Get().ToggleKeyAllEnabled, NAME_None, TAttribute<FText>(), TAttribute<FText>(), KeyAllIcon );
 		}
 
-		ToolBarBuilder.AddToolBarButton( FSequencerCommands::Get().ToggleAutoKeyEnabled );
+		TAttribute<FSlateIcon> AutoKeyModeIcon;
+		AutoKeyModeIcon.Bind(TAttribute<FSlateIcon>::FGetter::CreateLambda( [&] {
+			switch ( Sequencer.Pin()->GetAutoKeyMode() )
+			{
+			case EAutoKeyMode::KeyAll:
+				return FSequencerCommands::Get().SetAutoKeyModeAll->GetIcon();
+			case EAutoKeyMode::KeyAnimated:
+				return FSequencerCommands::Get().SetAutoKeyModeAnimated->GetIcon();
+			default: // EAutoKeyMode::KeyNone
+				return FSequencerCommands::Get().SetAutoKeyModeNone->GetIcon();
+			}
+		} ) );
+
+		TAttribute<FText> AutoKeyModeToolTip;
+		AutoKeyModeToolTip.Bind( TAttribute<FText>::FGetter::CreateLambda( [&] {
+			switch ( Sequencer.Pin()->GetAutoKeyMode() )
+			{
+			case EAutoKeyMode::KeyAll:
+				return FSequencerCommands::Get().SetAutoKeyModeAll->GetDescription();
+			case EAutoKeyMode::KeyAnimated:
+				return FSequencerCommands::Get().SetAutoKeyModeAnimated->GetDescription();
+			default: // EAutoKeyMode::KeyNone
+				return FSequencerCommands::Get().SetAutoKeyModeNone->GetDescription();
+			}
+		} ) );
+
+		ToolBarBuilder.AddComboButton(
+			FUIAction(),
+			FOnGetContent::CreateSP(this, &SSequencer::MakeAutoKeyMenu),
+			LOCTEXT("AutoKeyMode", "Auto-Key Mode"),
+			AutoKeyModeToolTip,
+			AutoKeyModeIcon);
 	}
 	ToolBarBuilder.EndSection();
 
@@ -825,7 +867,7 @@ TSharedRef<SWidget> SSequencer::MakeGeneralMenu()
 				.OnBeginSliderMovement(OnBeginPlaybackRangeDrag)
 				.OnEndSliderMovement_Lambda([=](float Value){ OnStartChanged(Value); OnEndPlaybackRangeDrag.ExecuteIfBound(); })
 				.MinValue_Lambda([=]() -> float {
-					return PinnedSequencer->GetViewRange().GetLowerBoundValue(); 
+					return PinnedSequencer->GetClampRange().GetLowerBoundValue(); 
 				})
 				.MaxValue_Lambda([=]() -> float {
 					return PinnedSequencer->GetPlaybackRange().GetUpperBoundValue(); 
@@ -861,7 +903,7 @@ TSharedRef<SWidget> SSequencer::MakeGeneralMenu()
 					return PinnedSequencer->GetPlaybackRange().GetLowerBoundValue(); 
 				})
 				.MaxValue_Lambda([=]() -> float {
-					return PinnedSequencer->GetViewRange().GetUpperBoundValue(); 
+					return PinnedSequencer->GetClampRange().GetUpperBoundValue(); 
 				})
 				.Value_Lambda([=]() -> float {
 					return PinnedSequencer->GetPlaybackRange().GetUpperBoundValue();
@@ -870,6 +912,7 @@ TSharedRef<SWidget> SSequencer::MakeGeneralMenu()
 			LOCTEXT("PlaybackStartEnd", "End"));
 
 		MenuBuilder.AddMenuEntry( FSequencerCommands::Get().ToggleKeepCursorInPlaybackRange );
+		MenuBuilder.AddMenuEntry( FSequencerCommands::Get().ToggleKeepPlaybackRangeInSectionBounds );
 	}
 
 
@@ -919,6 +962,19 @@ TSharedRef<SWidget> SSequencer::MakeSnapMenu()
 }
 
 
+TSharedRef<SWidget> SSequencer::MakeAutoKeyMenu()
+{
+	FMenuBuilder MenuBuilder(false, Sequencer.Pin()->GetCommandBindings());
+
+	MenuBuilder.AddMenuEntry(FSequencerCommands::Get().SetAutoKeyModeAll);
+	MenuBuilder.AddMenuEntry(FSequencerCommands::Get().SetAutoKeyModeAnimated);
+	MenuBuilder.AddMenuEntry(FSequencerCommands::Get().SetAutoKeyModeNone);
+
+	return MenuBuilder.MakeWidget();
+
+}
+
+
 TSharedRef<SWidget> SSequencer::MakeTransportControls()
 {
 	FEditorWidgetsModule& EditorWidgetsModule = FModuleManager::Get().LoadModuleChecked<FEditorWidgetsModule>( "EditorWidgets" );
@@ -1044,8 +1100,19 @@ void SSequencer::ResetBreadcrumbs()
 
 void SSequencer::OnOutlinerSearchChanged( const FText& Filter )
 {
-	SequencerNodeTree->FilterNodes( Filter.ToString() );
+	const FString FilterString = Filter.ToString();
+
+	SequencerNodeTree->FilterNodes(FilterString);
 	TreeView->Refresh();
+
+	if (FilterString.StartsWith(TEXT("label:")))
+	{
+		LabelBrowser->SetSelectedLabel(FilterString.RightChop(6));
+	}
+	else
+	{
+		LabelBrowser->SetSelectedLabel(FString());
+	}
 }
 
 
diff --git a/Engine/Source/Editor/Sequencer/Private/SSequencer.h b/Engine/Source/Editor/Sequencer/Private/SSequencer.h
--- a/Engine/Source/Editor/Sequencer/Private/SSequencer.h
+++ b/Engine/Source/Editor/Sequencer/Private/SSequencer.h
@@ -258,6 +258,9 @@ private:
 	/** Makes the snapping menu for the toolbar. */
 	TSharedRef<SWidget> MakeSnapMenu();
 
+	/** Makes the auto-key menu for the toolbar. */
+	TSharedRef<SWidget> MakeAutoKeyMenu();
+
 	/** Makes and configures a set of the standard UE transport controls. */
 	TSharedRef<SWidget> MakeTransportControls();
 
diff --git a/Engine/Source/Editor/Sequencer/Private/SSequencerLabelBrowser.cpp b/Engine/Source/Editor/Sequencer/Private/SSequencerLabelBrowser.cpp
--- a/Engine/Source/Editor/Sequencer/Private/SSequencerLabelBrowser.cpp
+++ b/Engine/Source/Editor/Sequencer/Private/SSequencerLabelBrowser.cpp
@@ -43,6 +43,25 @@ void SSequencerLabelBrowser::Construct(const FArguments& InArgs, TSharedRef<FSeq
 END_SLATE_FUNCTION_BUILD_OPTIMIZATION
 
 
+void SSequencerLabelBrowser::SetSelectedLabel(const FString& Label)
+{
+	if (!Label.IsEmpty())
+	{
+		for (auto& Node : LabelList)
+		{
+			if (Node->Label == Label)
+			{
+				LabelTreeView->SetSelection(Node);
+
+				return;
+			}
+		}
+	}
+
+	LabelTreeView->ClearSelection();
+}
+
+
 /* SSequencerLabelBrowser implementation
  *****************************************************************************/
 
@@ -111,6 +130,15 @@ void SSequencerLabelBrowser::ReloadLabelList(bool FullyReload)
 /* SSequencerLabelBrowser callbacks
  *****************************************************************************/
 
+void SSequencerLabelBrowser::HandleLabelListRowLabelRenamed(TSharedPtr<FSequencerLabelTreeNode> Node, const FString& NewLabel)
+{
+	if (Sequencer->GetLabelManager().RenameLabel(Node->Label, NewLabel))
+	{
+		ReloadLabelList(true);
+	}
+}
+
+
 void SSequencerLabelBrowser::HandleLabelManagerLabelsChanged()
 {
 	ReloadLabelList(true);
@@ -121,20 +149,6 @@ TSharedPtr<SWidget> SSequencerLabelBrowser::HandleLabelTreeViewContextMenuOpenin
 {
 	FMenuBuilder MenuBuilder(true /*bInShouldCloseWindowAfterMenuSelection*/, nullptr);
 	{
-		MenuBuilder.BeginSection("Customize", LOCTEXT("CustomizeContextMenuSectionName", "Customize"));
-		{
-			MenuBuilder.AddMenuEntry(
-				LOCTEXT("SetColorMenuEntryLabel", "Set Color"),
-				LOCTEXT("SetColorMenuEntryTip", "Set the background color of this label"),
-				FSlateIcon(),
-				FUIAction(
-					FExecuteAction::CreateSP(this, &SSequencerLabelBrowser::HandleSetColorMenuEntryExecute),
-					FCanExecuteAction::CreateSP(this, &SSequencerLabelBrowser::HandleSetColorMenuEntryCanExecute)
-				)
-			);
-		}
-		MenuBuilder.EndSection();
-
 		MenuBuilder.BeginSection("Edit", LOCTEXT("EditContextMenuSectionName", "Edit"));
 		{
 			MenuBuilder.AddMenuEntry(
@@ -167,7 +181,8 @@ TSharedPtr<SWidget> SSequencerLabelBrowser::HandleLabelTreeViewContextMenuOpenin
  TSharedRef<ITableRow> SSequencerLabelBrowser::HandleLabelTreeViewGenerateRow(TSharedPtr<FSequencerLabelTreeNode> Item, const TSharedRef<STableViewBase>& OwnerTable)
 {
 	return SNew(SSequencerLabelListRow, OwnerTable)
-		.Node(Item);
+		.Node(Item)
+		.OnLabelRenamed(this, &SSequencerLabelBrowser::HandleLabelListRowLabelRenamed);
 }
 
 
@@ -217,23 +232,24 @@ bool SSequencerLabelBrowser::HandleRemoveLabelMenuEntryCanExecute() const
 
 void SSequencerLabelBrowser::HandleRenameLabelMenuEntryExecute()
 {
-}
-
-
-bool SSequencerLabelBrowser::HandleRenameLabelMenuEntryCanExecute() const
-{
-	return false;
-}
+	TArray<TSharedPtr<FSequencerLabelTreeNode>> SelectedItems;
 
+	if (LabelTreeView->GetSelectedItems(SelectedItems) > 0)
+	{
+		auto ListRow = StaticCastSharedPtr<SSequencerLabelListRow>(LabelTreeView->WidgetFromItem(SelectedItems[0]));
 
-void SSequencerLabelBrowser::HandleSetColorMenuEntryExecute()
-{
+		if (ListRow.IsValid())
+		{
+			ListRow->EnterRenameMode();
+		}
+	}
 }
 
 
-bool SSequencerLabelBrowser::HandleSetColorMenuEntryCanExecute() const
+bool SSequencerLabelBrowser::HandleRenameLabelMenuEntryCanExecute() const
 {
-	return false;
+	TArray<TSharedPtr<FSequencerLabelTreeNode>> SelectedItems;
+	return (LabelTreeView->GetSelectedItems(SelectedItems) > 0);
 }
 
 
diff --git a/Engine/Source/Editor/Sequencer/Private/SSequencerLabelBrowser.h b/Engine/Source/Editor/Sequencer/Private/SSequencerLabelBrowser.h
--- a/Engine/Source/Editor/Sequencer/Private/SSequencerLabelBrowser.h
+++ b/Engine/Source/Editor/Sequencer/Private/SSequencerLabelBrowser.h
@@ -38,6 +38,13 @@ public:
 	 */
 	void Construct(const FArguments& InArgs, TSharedRef<FSequencer> InSequencer);
 
+	/**
+	 * Sets the selected label.
+	 *
+	 * @param Label The label to select.
+	 */
+	void SetSelectedLabel(const FString& Label);
+
 protected:
 
 	/** Reloads the list of processes. */
@@ -48,6 +55,9 @@ private:
 	/** Callback for when the label manager's labels changed. */
 	void HandleLabelManagerLabelsChanged();
 
+	/** Callback for when a label has been renamed in the label tree view. */
+	void HandleLabelListRowLabelRenamed(TSharedPtr<FSequencerLabelTreeNode> Node, const FString& NewLabel);
+
 	/** Callback for generating the label tree view's context menu. */
 	TSharedPtr<SWidget> HandleLabelTreeViewContextMenuOpening();
 
@@ -72,12 +82,6 @@ private:
 	/** Callback for checking whether 'Rename' context menu entry can execute. */
 	bool HandleRenameLabelMenuEntryCanExecute() const;
 
-	/** Callback for executing the 'Set Color' context menu entry. */
-	void HandleSetColorMenuEntryExecute();
-
-	/** Callback for checking whether 'Set Color' context menu entry can execute. */
-	bool HandleSetColorMenuEntryCanExecute() const;
-
 private:
 
 	/** Holds the collection of root labels to be displayed in the tree view. */
diff --git a/Engine/Source/Editor/Sequencer/Private/SSequencerLabelListRow.h b/Engine/Source/Editor/Sequencer/Private/SSequencerLabelListRow.h
--- a/Engine/Source/Editor/Sequencer/Private/SSequencerLabelListRow.h
+++ b/Engine/Source/Editor/Sequencer/Private/SSequencerLabelListRow.h
@@ -2,7 +2,16 @@
 
 #pragma once
 
+#include "SEditableLabel.h"
 
+
+/** Delegate that is executed whenever a label has been renamed. */
+DECLARE_DELEGATE_TwoParams(FOnSequencerLabelRenamed, TSharedPtr<FSequencerLabelTreeNode> /*Node*/, const FString& /*NewLabel*/);
+
+
+/**
+ * Represents a node in the label tree.
+ */
 struct FSequencerLabelTreeNode
 {
 public:
@@ -41,7 +50,11 @@ class SSequencerLabelListRow
 public:
 
 	SLATE_BEGIN_ARGS(SSequencerLabelListRow) { }
+		/** The label tree node data visualized in this list row. */
 		SLATE_ARGUMENT(TSharedPtr<FSequencerLabelTreeNode>, Node)
+
+		/** Called whenever the folder has been renamed. */
+		SLATE_EVENT(FOnSequencerLabelRenamed, OnLabelRenamed)
 	SLATE_END_ARGS()
 
 public:
@@ -54,6 +67,7 @@ public:
 	void Construct( const FArguments& InArgs, const TSharedRef<STableViewBase>& InOwnerTableView )
 	{
 		Node = InArgs._Node;
+		OnLabelRenamed = InArgs._OnLabelRenamed;
 
 		STableRow<TSharedPtr<FSequencerLabelTreeNode>>::Construct(
 			STableRow<TSharedPtr<FSequencerLabelTreeNode>>::FArguments()
@@ -79,26 +93,24 @@ public:
 						.Padding(0.0f, 2.0f)
 						.VAlign(VAlign_Center)
 						[
-							SNew(STextBlock)
+							SAssignNew(EditableLabel, SEditableLabel)
+								.CanEdit(this, &SSequencerLabelListRow::HandleFolderNameCanEdit)
+								.OnTextChanged(this, &SSequencerLabelListRow::HandleFolderNameTextChanged)
 								.Text(
 									Node->Label.IsEmpty()
 										? LOCTEXT("AllTracksLabel", "All Tracks")
-										: Node->DisplayName
+										: FText::FromString(Node->Label)
 								)
 						]
-
-					// edit icon
-					+ SHorizontalBox::Slot()
-						.AutoWidth()
-						.VAlign(VAlign_Center)
-						[
-							SNew(SImage)
-								.Image(FCoreStyle::Get().GetBrush(TEXT("EditableLabel.EditIcon")))
-								.Visibility(this, &SSequencerLabelListRow::HandleEditIconVisibility)
-						]
 				],
-			InOwnerTableView
-		);
+				InOwnerTableView
+			);
+	}
+
+	/** Change the label text to edit mode. */
+	void EnterRenameMode()
+	{
+		EditableLabel->EnterTextMode();
 	}
 
 private:
@@ -126,10 +138,31 @@ private:
 		return FLinearColor::Gray;
 	}
 
+	bool HandleFolderNameCanEdit() const
+	{
+		return !Node->Label.IsEmpty();
+	}
+
+	void HandleFolderNameTextChanged(const FText& NewLabel)
+	{
+		FString NewLabelString = NewLabel.ToString();
+
+		if (NewLabelString != Node->Label)
+		{
+			OnLabelRenamed.ExecuteIfBound(Node.ToSharedRef(), NewLabelString);
+		}
+	}
+
 private:
 
+	/** Holds the editable text label widget. */
+	TSharedPtr<SEditableLabel> EditableLabel;
+
 	/** Holds the label node. */
 	TSharedPtr<FSequencerLabelTreeNode> Node;
+
+	/** A delegate that is executed whenever the label has been renamed. */
+	FOnSequencerLabelRenamed OnLabelRenamed;
 };
 
 
diff --git a/Engine/Source/Editor/Sequencer/Private/SSequencerTrackArea.cpp b/Engine/Source/Editor/Sequencer/Private/SSequencerTrackArea.cpp
--- a/Engine/Source/Editor/Sequencer/Private/SSequencerTrackArea.cpp
+++ b/Engine/Source/Editor/Sequencer/Private/SSequencerTrackArea.cpp
@@ -127,7 +127,7 @@ int32 SSequencerTrackArea::OnPaint(const FPaintArgs& Args, const FGeometry& Allo
 
 	for (const auto& TrackEditor : TrackEditors)
 	{
-		LayerId = TrackEditor->PaintTrackArea(Args, AllottedGeometry, MyClippingRect, OutDrawElements, LayerId, InWidgetStyle);
+		LayerId = TrackEditor->PaintTrackArea(Args, AllottedGeometry, MyClippingRect, OutDrawElements, LayerId + 1, InWidgetStyle);
 	}
 
 	// paint the child widgets
@@ -138,12 +138,12 @@ int32 SSequencerTrackArea::OnPaint(const FPaintArgs& Args, const FGeometry& Allo
 	{
 		FArrangedWidget& CurWidget = ArrangedChildren[ChildIndex];
 		FSlateRect ChildClipRect = MyClippingRect.IntersectionWith( CurWidget.Geometry.GetClippingRect() );
-		const int32 ThisWidgetLayerId = CurWidget.Widget->Paint( Args.WithNewParent(this), CurWidget.Geometry, ChildClipRect, OutDrawElements, LayerId + 1, InWidgetStyle, ShouldBeEnabled( bParentEnabled ) );
+		const int32 ThisWidgetLayerId = CurWidget.Widget->Paint( Args.WithNewParent(this), CurWidget.Geometry, ChildClipRect, OutDrawElements, LayerId + 2, InWidgetStyle, ShouldBeEnabled( bParentEnabled ) );
 
 		LayerId = FMath::Max(LayerId, ThisWidgetLayerId);
 	}
 
-	return Sequencer->GetEditTool().OnPaint(AllottedGeometry, MyClippingRect, OutDrawElements, LayerId + 1);
+	return Sequencer->GetEditTool().OnPaint(AllottedGeometry, MyClippingRect, OutDrawElements, LayerId + 2);
 }
 
 
@@ -157,6 +157,8 @@ FReply SSequencerTrackArea::OnMouseButtonDown( const FGeometry& MyGeometry, cons
 
 FReply SSequencerTrackArea::OnMouseButtonUp( const FGeometry& MyGeometry, const FPointerEvent& MouseEvent )
 {
+	FContextMenuSuppressor SuppressContextMenus(TimeSliderController.ToSharedRef());
+
 	// Always ensure the edit tool is set up
 	InputStack.SetHandlerAt(0, &Sequencer->GetEditTool());
 	return InputStack.HandleMouseButtonUp(*this, MyGeometry, MouseEvent);
