Movie Scene Capture: Allow game sounds if capturing audio

diff --git a/Engine/Source/Editor/MovieSceneCaptureDialog/Private/MovieSceneCaptureDialogModule.cpp b/Engine/Source/Editor/MovieSceneCaptureDialog/Private/MovieSceneCaptureDialogModule.cpp
--- a/Engine/Source/Editor/MovieSceneCaptureDialog/Private/MovieSceneCaptureDialogModule.cpp
+++ b/Engine/Source/Editor/MovieSceneCaptureDialog/Private/MovieSceneCaptureDialogModule.cpp
@@ -64,6 +64,7 @@
 
 #include "GameFramework/WorldSettings.h"
 #include "FrameNumberNumericInterface.h"
+#include "Protocols/AudioCaptureProtocol.h"
 
 #define LOCTEXT_NAMESPACE "MovieSceneCaptureDialog"
 
@@ -475,7 +476,7 @@ void FInEditorCapture::OverridePlaySettings(ULevelEditorPlaySettings* PlayInEdit
 	PlayInEditorSettings->ShowMouseControlLabel = false;
 	PlayInEditorSettings->ViewportGetsHMDControl = false;
 	PlayInEditorSettings->ShouldMinimizeEditorOnVRPIE = true;
-	PlayInEditorSettings->EnableGameSound = false;
+	PlayInEditorSettings->EnableGameSound = CaptureObject->AudioCaptureProtocolType != UNullAudioCaptureProtocol::StaticClass();
 	PlayInEditorSettings->bOnlyLoadVisibleLevelsInPIE = false;
 	PlayInEditorSettings->bPreferToStreamLevelsInPIE = false;
 	PlayInEditorSettings->PIEAlwaysOnTop = false;
