Fix compilation issue on Linux.

diff --git a/Engine/Source/Editor/MovieSceneCaptureDialog/Public/MovieSceneCaptureDialogModule.h b/Engine/Source/Editor/MovieSceneCaptureDialog/Public/MovieSceneCaptureDialogModule.h
--- a/Engine/Source/Editor/MovieSceneCaptureDialog/Public/MovieSceneCaptureDialogModule.h
+++ b/Engine/Source/Editor/MovieSceneCaptureDialog/Public/MovieSceneCaptureDialogModule.h
@@ -131,7 +131,7 @@ private:
 	TArray<uint8> BackedUpPlaySettings;
 
 	bool CachedPathTracingMode = false;
-	FEngineShowFlags* CachedEngineShowFlags = nullptr;
+	struct FEngineShowFlags* CachedEngineShowFlags = nullptr;
 	TSubclassOf<AGameModeBase> CachedGameMode;
 };
 
