Sequencer: Texture streaming toggle also sets r.TextureStreaming

diff --git a/Engine/Source/Editor/MovieSceneCaptureDialog/Private/MovieSceneCaptureDialogModule.cpp b/Engine/Source/Editor/MovieSceneCaptureDialog/Private/MovieSceneCaptureDialogModule.cpp
--- a/Engine/Source/Editor/MovieSceneCaptureDialog/Private/MovieSceneCaptureDialogModule.cpp
+++ b/Engine/Source/Editor/MovieSceneCaptureDialog/Private/MovieSceneCaptureDialogModule.cpp
@@ -409,13 +409,20 @@ void FInEditorCapture::Start()
 			CVarStreamingPoolSize->Set(UndefinedTexturePoolSize, ECVF_SetByConsole);
 		}
 
-			IConsoleVariable* CVarUseFixedPoolSize = IConsoleManager::Get().FindConsoleVariable(TEXT("r.Streaming.UseFixedPoolSize"));
-			if (CVarUseFixedPoolSize)
-			{
-				BackedUpUseFixedPoolSize = CVarUseFixedPoolSize->GetInt(); 
-				CVarUseFixedPoolSize->Set(0, ECVF_SetByConsole);
-			}
+		IConsoleVariable* CVarUseFixedPoolSize = IConsoleManager::Get().FindConsoleVariable(TEXT("r.Streaming.UseFixedPoolSize"));
+		if (CVarUseFixedPoolSize)
+		{
+			BackedUpUseFixedPoolSize = CVarUseFixedPoolSize->GetInt(); 
+			CVarUseFixedPoolSize->Set(0, ECVF_SetByConsole);
+		}
+
+		IConsoleVariable* CVarTextureStreaming = IConsoleManager::Get().FindConsoleVariable(TEXT("r.TextureStreaming"));
+		if (CVarTextureStreaming)
+		{
+			BackedUpTextureStreaming = CVarTextureStreaming->GetInt();
+			CVarTextureStreaming->Set(0, ECVF_SetByConsole);
 		}
+	}
 
 	FObjectWriter(PlayInEditorSettings, BackedUpPlaySettings);
 	OverridePlaySettings(PlayInEditorSettings);
@@ -560,6 +567,12 @@ void FInEditorCapture::Shutdown()
 		{
 			CVarUseFixedPoolSize->Set(BackedUpUseFixedPoolSize, ECVF_SetByConsole);
 		}
+
+		IConsoleVariable* CVarTextureStreaming = IConsoleManager::Get().FindConsoleVariable(TEXT("r.TextureStreaming"));
+		if (CVarTextureStreaming)
+		{
+			CVarTextureStreaming->Set(BackedUpTextureStreaming, ECVF_SetByConsole);
+		}
 	}
 
 	if (CaptureObject->Settings.GameModeOverride != nullptr && CapturingFromWorld != nullptr)
diff --git a/Engine/Source/Editor/MovieSceneCaptureDialog/Public/MovieSceneCaptureDialogModule.h b/Engine/Source/Editor/MovieSceneCaptureDialog/Public/MovieSceneCaptureDialogModule.h
--- a/Engine/Source/Editor/MovieSceneCaptureDialog/Public/MovieSceneCaptureDialogModule.h
+++ b/Engine/Source/Editor/MovieSceneCaptureDialog/Public/MovieSceneCaptureDialogModule.h
@@ -128,6 +128,7 @@ private:
 	float TransientMasterVolume;
 	int32 BackedUpStreamingPoolSize;
 	int32 BackedUpUseFixedPoolSize;
+	int32 BackedUpTextureStreaming;
 	TArray<uint8> BackedUpPlaySettings;
 
 	TSubclassOf<AGameModeBase> CachedGameMode;
