More fix-ups for editor calls to deprecated audio device calls

diff --git a/Engine/Source/Editor/MovieSceneCaptureDialog/Private/MovieSceneCaptureDialogModule.cpp b/Engine/Source/Editor/MovieSceneCaptureDialog/Private/MovieSceneCaptureDialogModule.cpp
--- a/Engine/Source/Editor/MovieSceneCaptureDialog/Private/MovieSceneCaptureDialogModule.cpp
+++ b/Engine/Source/Editor/MovieSceneCaptureDialog/Private/MovieSceneCaptureDialogModule.cpp
@@ -433,9 +433,8 @@ void FInEditorCapture::Start()
 
 	UGameViewportClient::OnViewportCreated().AddRaw(this, &FInEditorCapture::OnPIEViewportStarted);
 	FEditorDelegates::EndPIE.AddRaw(this, &FInEditorCapture::OnEndPIE);
-		
-	FAudioDevice* AudioDevice = GEngine->GetMainAudioDevice();
-	if (AudioDevice != nullptr)
+
+	if (FAudioDevice* AudioDevice = FAudioDeviceManager::GetMainDevice())
 	{
 		TransientMasterVolume = AudioDevice->GetTransientMasterVolume();
 		AudioDevice->SetTransientMasterVolume(0.0f);
@@ -596,8 +595,7 @@ void FInEditorCapture::Shutdown()
 
 	FObjectReader(GetMutableDefault<ULevelEditorPlaySettings>(), BackedUpPlaySettings);
 
-	FAudioDevice* AudioDevice = GEngine->GetMainAudioDevice();
-	if (AudioDevice != nullptr)
+	if (FAudioDevice* AudioDevice = FAudioDeviceManager::GetMainDevice())
 	{
 		AudioDevice->SetTransientMasterVolume(TransientMasterVolume);
 	}
diff --git a/Engine/Source/Editor/MovieSceneTools/Private/TrackEditors/AudioTrackEditor.cpp b/Engine/Source/Editor/MovieSceneTools/Private/TrackEditors/AudioTrackEditor.cpp
--- a/Engine/Source/Editor/MovieSceneTools/Private/TrackEditors/AudioTrackEditor.cpp
+++ b/Engine/Source/Editor/MovieSceneTools/Private/TrackEditors/AudioTrackEditor.cpp
@@ -296,8 +296,7 @@ void FAudioThumbnail::GenerateWaveformPreview(TArray<uint8>& OutData, TRange<flo
 	{
 		// @todo Sequencer optimize - We might want to generate the data when we generate the texture
 		// and then discard the data afterwards, though that might be a perf hit traded for better memory usage
-		FAudioDevice* AudioDevice = GEngine->GetMainAudioDevice();
-		if (AudioDevice)
+		if (FAudioDevice* AudioDevice = FAudioDeviceManager::GetMainDevice())
 		{
 			AudioDevice->StopAllSounds(true);
 
